Move files prior to NetCore checkin.

diff --git a/Engine/Source/Runtime/Net/Core/Classes/Net/Core/Analytics/NetAnalyticsAggregatorConfig.h b/Engine/Source/Runtime/Net/Core/Classes/Net/Core/Analytics/NetAnalyticsAggregatorConfig.h
new file mode 100644
--- /dev/null
+++ b/Engine/Source/Runtime/Net/Core/Classes/Net/Core/Analytics/NetAnalyticsAggregatorConfig.h
@@ -0,0 +1,46 @@
+// Copyright 1998-2019 Epic Games, Inc. All Rights Reserved.
+
+#pragma once
+
+#include "CoreMinimal.h"
+#include "UObject/ObjectMacros.h"
+#include "UObject/Object.h"
+
+#include "NetAnalyticsAggregatorConfig.generated.h"
+
+
+/**
+ * Configuration for FNetAnalyticsData - enabling/disabling analytics data, based on DataName
+ */
+USTRUCT()
+struct FNetAnalyticsDataConfig
+{
+	GENERATED_USTRUCT_BODY()
+
+	/** The name of the analytics data type (should match analytics stat name) */
+	UPROPERTY(config)
+	FName DataName;
+
+	/** Whether or not the specified analytics data type, is enabled */
+	UPROPERTY(config)
+	bool bEnabled;
+};
+
+/**
+ * Configuration for FNetAnalyticsAggregator - loaded PerObjectConfig, for each NetDriverName
+ */
+UCLASS(config=Engine, PerObjectConfig)
+class UNetAnalyticsAggregatorConfig : public UObject
+{
+	GENERATED_UCLASS_BODY()
+
+protected:
+	void OverridePerObjectConfigSection(FString& SectionName) override;
+
+public:
+	/** Registers FNetAnalyticsData data holders, by DataName - and specifies whether they are enabled or disabled */
+	UPROPERTY(config)
+	TArray<FNetAnalyticsDataConfig> NetAnalyticsData;
+};
+
+
diff --git a/Engine/Source/Runtime/Net/Core/Private/Net/Core/Analytics/NetAnalytics.cpp b/Engine/Source/Runtime/Net/Core/Private/Net/Core/Analytics/NetAnalytics.cpp
new file mode 100644
--- /dev/null
+++ b/Engine/Source/Runtime/Net/Core/Private/Net/Core/Analytics/NetAnalytics.cpp
@@ -0,0 +1,144 @@
+// Copyright 1998-2019 Epic Games, Inc. All Rights Reserved.
+
+// Includes
+#include "NetAnalytics.h"
+#include "PacketHandler.h"
+#include "NetAnalyticsAggregatorConfig.h"
+
+
+// Globals
+
+TAtomic<uint8> GNetAnalyticsCounter(0);
+
+
+/**
+ * FNetAnalytics
+ */
+
+void FNetAnalyticsData::InternalSendAnalytics()
+{
+	SendAnalytics();
+}
+
+#if NET_ANALYTICS_MULTITHREADING
+/**
+ * FThreadedNetAnalyticsData
+ */
+
+FThreadedNetAnalyticsData::FThreadedNetAnalyticsData()
+	: FNetAnalyticsData()
+	, bReadyToSend(false)
+{
+}
+
+void FThreadedNetAnalyticsData::InternalSendAnalytics()
+{
+	bReadyToSend = true;
+}
+
+void FThreadedNetAnalyticsData::NotifyFinalRelease()
+{
+	if (bReadyToSend)
+	{
+		SendAnalytics();
+	}
+}
+#endif
+
+
+/**
+ * FNetAnalyticsAggregator
+ */
+
+FNetAnalyticsAggregator::FNetAnalyticsAggregator(TSharedPtr<IAnalyticsProvider> InProvider, FName InNetDriverName)
+	: AnalyticsProvider(InProvider)
+	, NetDriverName(InNetDriverName)
+	, AnalyticsDataMap()
+	, AnalyticsDataTypeMap()
+	, AnalyticsDataConfigMap()
+	, bSentAnalytics(false)
+{
+}
+
+void FNetAnalyticsAggregator::Init()
+{
+	GNetAnalyticsCounter++;
+
+	InitConfig();
+}
+
+void FNetAnalyticsAggregator::InitConfig()
+{
+	UClass* ClassRef = UNetAnalyticsAggregatorConfig::StaticClass();
+	UNetAnalyticsAggregatorConfig* CurConfig = FindObject<UNetAnalyticsAggregatorConfig>(ClassRef, *NetDriverName.ToString());
+
+	if (CurConfig == nullptr)
+	{
+		CurConfig = NewObject<UNetAnalyticsAggregatorConfig>(ClassRef, NetDriverName);
+	}
+
+	AnalyticsDataConfigMap.Empty();
+
+	// If the config is hotfixed, make sure no data holders are currently active, as they can't be selectively hotfixed if loaded
+	// (this does seem to happen, frequently - so limits the hotfixability of Net Analytics)
+	UE_CLOG(AnalyticsDataMap.Num() > 0, PacketHandlerLog, Warning,
+			TEXT("Net Analytics hotfixed while already active. Analytics hotfix changes may not be applied correctly."));
+
+	if (CurConfig != nullptr)
+	{
+		for (const FNetAnalyticsDataConfig& CurEntry : CurConfig->NetAnalyticsData)
+		{
+			AnalyticsDataConfigMap.Add(CurEntry.DataName, CurEntry.bEnabled);
+
+			UE_LOG(PacketHandlerLog, Log, TEXT("Adding NetAnalyticsData: %s, bEnabled: %i"), *CurEntry.DataName.ToString(),
+					CurEntry.bEnabled);
+		}
+	}
+}
+
+TNetAnalyticsDataPtr<> FNetAnalyticsAggregator::RegisterAnalyticsData_Internal(TNetAnalyticsDataRef<> InData, const FName& InDataName,
+																					FString InTypeName)
+{
+	TNetAnalyticsDataPtr<> Result;
+	bool* bEnabled = AnalyticsDataConfigMap.Find(InDataName);
+
+	if (bEnabled != nullptr && *bEnabled)
+	{
+		FString* TypeVal = AnalyticsDataTypeMap.Find(InDataName);
+
+		if (TypeVal == nullptr)
+		{
+			AnalyticsDataTypeMap.Add(InDataName, InTypeName);
+		}
+		else
+		{
+			check(*TypeVal == InTypeName);
+		}
+
+		TNetAnalyticsDataRef<>* Found = AnalyticsDataMap.Find(InDataName);
+
+		Result = Found != nullptr ? *Found : AnalyticsDataMap.Add(InDataName, InData);
+
+		InData->Aggregator = this;
+	}
+	else if (bEnabled == nullptr)
+	{
+		UE_LOG(PacketHandlerLog, Error, TEXT("NetAnalyticsData type '%s' must be added to NetAnalyticsAggregatorConfig, for NetDriverName: %s."),
+				*InDataName.ToString(), *NetDriverName.ToString());
+	}
+
+	return Result;
+}
+
+void FNetAnalyticsAggregator::SendAnalytics()
+{
+	if (!bSentAnalytics)
+	{
+		for (TMap<FName, TNetAnalyticsDataRef<>>::TConstIterator It(AnalyticsDataMap); It; ++It)
+		{
+			It.Value()->InternalSendAnalytics();
+		}
+
+		bSentAnalytics = true;
+	}
+}
diff --git a/Engine/Source/Runtime/Net/Core/Private/Net/Core/Analytics/NetAnalyticsAggregatorConfig.cpp b/Engine/Source/Runtime/Net/Core/Private/Net/Core/Analytics/NetAnalyticsAggregatorConfig.cpp
new file mode 100644
--- /dev/null
+++ b/Engine/Source/Runtime/Net/Core/Private/Net/Core/Analytics/NetAnalyticsAggregatorConfig.cpp
@@ -0,0 +1,17 @@
+// Copyright 1998-2019 Epic Games, Inc. All Rights Reserved.
+
+#include "NetAnalyticsAggregatorConfig.h"
+
+
+/**
+ * UNetAnalyticsAggregatorConfig
+ */
+UNetAnalyticsAggregatorConfig::UNetAnalyticsAggregatorConfig(const FObjectInitializer& ObjectInitializer)
+	: Super(ObjectInitializer)
+{
+}
+
+void UNetAnalyticsAggregatorConfig::OverridePerObjectConfigSection(FString& SectionName)
+{
+	SectionName = GetName() + TEXT(" ") + GetClass()->GetName();
+}
diff --git a/Engine/Source/Runtime/Net/Core/Private/Net/Core/Misc/DDoSDetection.cpp b/Engine/Source/Runtime/Net/Core/Private/Net/Core/Misc/DDoSDetection.cpp
new file mode 100644
--- /dev/null
+++ b/Engine/Source/Runtime/Net/Core/Private/Net/Core/Misc/DDoSDetection.cpp
@@ -0,0 +1,375 @@
+// Copyright 1998-2019 Epic Games, Inc. All Rights Reserved.
+
+// Includes
+
+#include "DDoSDetection.h"
+#include "PacketHandler.h"
+#include "Misc/ConfigCacheIni.h"
+
+
+/**
+ * DDoS Detection
+ *
+ * DDoS (Distributed Denial of Service) attacks typically hinder game servers by flooding them with so many packets,
+ * that they are unable to process all of the packets without locking up and/or drowning out other players packets,
+ * causing players to time out or to suffer severe packet loss which hinders gameplay.
+ *
+ * Typically these attacks use spoofed UDP packets, where the source IP is unverifiable,
+ * and so IP banning is usually not an effective or advisable means of blocking such attacks.
+ *
+ * This DDoS detection focuses specifically on this situation, detecting/mitigating DDoS attacks based on spoofed UDP packets,
+ * which do not originate from an existing NetConnection. Flooding attacks coming from an existing NetConnection are a separate issue,
+ * as (due to the stateless handshake required before creating a NetConnection) the IP will be verified,
+ * and so such attacks should be dealt with through IP banning - this and other types of DoS attacks are not dealt with by this code.
+ *
+ *
+ * Implementation:
+ *
+ * DDoS attacks are detected by setting configurable thresholds for the number of packets per second,
+ * and milliseconds per frame spent processing packets, beyond which the DDoS detection will escalate to a higher severity state.
+ *
+ * Each severity state has a different set of thresholds before it will escalate to the next state,
+ * and can also place a limit on the number of packets processed per second, and/or milliseconds spent processing.
+ *
+ * The stronger the DDoS attack, the higher the severity state will escalate (based on the thresholds),
+ * and the stronger the limitations places on incoming packets will be, in order to try and maintain good server performance.
+ *
+ *
+ * Limitations:
+ *
+ * Heavy DDoS:
+ * While the code can withstand a heavy, locally hosted, multithreaded DDoS,
+ * past a certain point network hardware and bandwidth capacity will become a limit, and even with strong enough hardware,
+ * the OS kernel calls for receiving packets, will become a limit (for Linux, recvmmsg may be used to alleviate this later).
+ *
+ * So this code just deals with as much of the DDoS as it can, at an application level - if you're getting hit with a bad enough DDoS,
+ * then you're going to have to look at measures at the network infrastructure level - for example,
+ * IP filtering at the edge of your network, communicating with the game server to only allow packets from existing NetConnection IP's.
+ *
+ * Tuning thresholds per-Game:
+ * You will need to manually tune the packet thresholds specifically for your game, even for each different gametype within your game,
+ * and maybe even community server admins will need to retune, if hosting a server with mods etc..
+ *
+ * Blocking new connections:
+ * If a DDoS is expensive enough, that you choose to drop non-NetConnection packets after a threshold (a wise move, for performance...),
+ * then new players will be blocked from entering the server.
+ */
+
+// @todo #JohnB: The code deliberately counts the time spent processing NetConnection RPC's, not just merely receiving packets.
+//					Make sure this doesn't introduce problems or false positives (or at least, that they're tolerable, if it does).
+
+/**
+ * FDDoSDetection
+ */
+
+FDDoSDetection::FDDoSDetection()
+	: bDDoSDetection(false)
+	, bDDoSAnalytics(false)
+	, bHitFrameNonConnLimit(false)
+	, bHitFrameNetConnLimit(false)
+	, DetectionSeverity()
+	, ActiveState(0)
+	, WorstActiveState(0)
+	, LastMetEscalationConditions(0.0)
+	, bMetEscalationConditionsThisFrame(false)
+	, bDDoSLogRestrictions(false)
+	, DDoSLogSpamLimit(0)
+	, LogHitCounter(0)
+	, HitchTimeQuotaMS(-1)
+	, HitchFrameTolerance(-1)
+	, HitchFrameCount(0)
+	, LastPerSecQuotaBegin(0.0)
+	, CounterPerSecHistory()
+	, LastCounterPerSecHistoryIdx(0)
+	, StartFrameRecvTimestamp(0.0)
+	, EndFrameRecvTimestamp(0.0)
+	, StartFramePacketCount(0)
+	, ExpectedFrameTime(0.0)
+	, FrameAdjustment(0.f)
+	, NotifySeverityEscalation()
+{
+}
+
+void FDDoSDetection::Init(int32 MaxTickRate)
+{
+	ExpectedFrameTime = 1.0 / (MaxTickRate > 0.0 ? MaxTickRate : 30.0);
+
+	InitConfig();
+}
+
+void FDDoSDetection::InitConfig()
+{
+	const TCHAR* DDoSSection = TEXT("DDoSDetection");
+	int32 HitchFrameTolerance32 = -1;
+
+	GConfig->GetBool(DDoSSection, TEXT("bDDoSDetection"), bDDoSDetection, GEngineIni);
+	GConfig->GetBool(DDoSSection, TEXT("bDDoSAnalytics"), bDDoSAnalytics, GEngineIni);
+	GConfig->GetInt(DDoSSection, TEXT("DDoSLogSpamLimit"), DDoSLogSpamLimit, GEngineIni);
+	GConfig->GetInt(DDoSSection, TEXT("HitchTimeQuotaMS"), HitchTimeQuotaMS, GEngineIni);
+	GConfig->GetInt(DDoSSection, TEXT("HitchFrameTolerance"), HitchFrameTolerance32, GEngineIni);
+
+	HitchFrameTolerance = HitchFrameTolerance32;
+	DDoSLogSpamLimit = DDoSLogSpamLimit > 0 ? DDoSLogSpamLimit : 64;
+
+	DetectionSeverity.Empty();
+
+	UE_LOG(PacketHandlerLog, Log, TEXT("DDoS detection status: detection enabled: %d analytics enabled: %d"), bDDoSDetection, bDDoSAnalytics);
+
+	if (bDDoSDetection)
+	{
+		TArray<FString> SeverityCatagories;
+		int32 HighestCooloffTime = 0;
+
+		GConfig->GetArray(DDoSSection, TEXT("DetectionSeverity"), SeverityCatagories, GEngineIni);
+
+		for (const FString& CurCategory : SeverityCatagories)
+		{
+			FString CurSection = FString(DDoSSection) + TEXT(".") + CurCategory;
+
+			if (GConfig->DoesSectionExist(*CurSection, GEngineIni))
+			{
+				FDDoSStateConfig& CurState = DetectionSeverity.AddDefaulted_GetRef();
+				int32 EscalateTime32 = 0;
+
+				CurState.SeverityCategory = CurCategory;
+
+				GConfig->GetInt(*CurSection, TEXT("EscalateQuotaPacketsPerSec"), CurState.EscalateQuotaPacketsPerSec, GEngineIni);
+				GConfig->GetInt(*CurSection, TEXT("EscalateQuotaDisconnPacketsPerSec"), CurState.EscalateQuotaDisconnPacketsPerSec, GEngineIni);
+				GConfig->GetInt(*CurSection, TEXT("EscalateQuotaBadPacketsPerSec"), CurState.EscalateQuotaBadPacketsPerSec, GEngineIni);
+				GConfig->GetInt(*CurSection, TEXT("PacketLimitPerFrame"), CurState.PacketLimitPerFrame, GEngineIni);
+				GConfig->GetInt(*CurSection, TEXT("PacketTimeLimitMSPerFrame"), CurState.PacketTimeLimitMSPerFrame, GEngineIni);
+				GConfig->GetInt(*CurSection, TEXT("NetConnPacketTimeLimitMSPerFrame"), CurState.NetConnPacketTimeLimitMSPerFrame, GEngineIni);
+				GConfig->GetInt(*CurSection, TEXT("CooloffTime"), CurState.CooloffTime, GEngineIni);
+
+				if (GConfig->GetInt(*CurSection, TEXT("EscalateTimeQuotaMSPerFrame"), EscalateTime32, GEngineIni))
+				{
+					CurState.EscalateTimeQuotaMSPerFrame = EscalateTime32;
+				}
+
+				HighestCooloffTime = FMath::Max(HighestCooloffTime, CurState.CooloffTime);
+			}
+			else
+			{
+				UE_LOG(PacketHandlerLog, Warning, TEXT("DDoS detection could not find ini section: %s"), *CurSection);
+			}
+		}
+
+		if (DetectionSeverity.Num() > 0)
+		{
+			DetectionSeverity[ActiveState].ApplyState(*this);
+
+			CounterPerSecHistory.SetNum(HighestCooloffTime);
+		}
+		else
+		{
+			UE_LOG(PacketHandlerLog, Warning, TEXT("DDoS detection enabled, but no DetectionSeverity states specified! Disabling."));
+
+			bDDoSDetection = false;
+		}
+	}
+}
+
+void FDDoSDetection::UpdateSeverity(bool bEscalate)
+{
+	int32 NewState = FMath::Clamp(ActiveState + (bEscalate ? 1 : -1), 0, DetectionSeverity.Num());
+
+	if (NewState != ActiveState)
+	{
+		double CurTime = FPlatformTime::Seconds();
+
+		if (bEscalate)
+		{
+			LastMetEscalationConditions = CurTime;
+		}
+		else
+		{
+			// De-escalate to the lowest state which hasn't cooled off, and estimate the timestamp for when the cooloff was last reset
+			// (due to estimating, there is slight inaccuracy in the cooloff time)
+			bool bCooloffReached = true;
+
+			while (bCooloffReached && NewState > 0)
+			{
+				FDDoSStateConfig& PrevState = DetectionSeverity[NewState-1];
+				int32 CurStateCooloffTime = DetectionSeverity[NewState].CooloffTime;
+
+				check(CounterPerSecHistory.Num() >= CurStateCooloffTime);
+
+				for (int32 SecondsDelta=0; SecondsDelta<CurStateCooloffTime; SecondsDelta++)
+				{
+					int32 CurIdx = LastCounterPerSecHistoryIdx - SecondsDelta;
+
+					CurIdx = (CurIdx < 0 ? CounterPerSecHistory.Num() + CurIdx : CurIdx);
+
+					check(CurIdx >= 0 && CurIdx < CounterPerSecHistory.Num());
+
+					FDDoSPacketCounters& CurPerSecHistory = CounterPerSecHistory[CurIdx];
+
+					if (PrevState.HasHitQuota(CurPerSecHistory, CurPerSecHistory.WorstFrameReceiveTimeMS))
+					{
+						// The state we're transitioning down into, would have last had its cooloff reset around this time
+						LastMetEscalationConditions = CurTime - (double)SecondsDelta;
+
+						bCooloffReached = false;
+						break;
+					}
+				}
+
+				if (bCooloffReached)
+				{
+					NewState--;
+				}
+			}
+		}
+
+
+		FDDoSStateConfig& OldState = DetectionSeverity[ActiveState];
+		FDDoSStateConfig& CurState = DetectionSeverity[NewState];
+
+		// If we're at anything other than the base state, then disable all unnecessary logs
+		bDDoSLogRestrictions = NewState > 0;
+		ActiveState = NewState;
+		bMetEscalationConditionsThisFrame = false;
+
+		CurState.ApplyState(*this);
+
+		if (FrameAdjustment > 0.f)
+		{
+			CurState.ApplyAdjustedState(*this, FrameAdjustment);
+		}
+
+
+		UE_LOG(PacketHandlerLog, Warning, TEXT("Updated DDoS detection severity from '%s' to '%s'"),
+				*OldState.SeverityCategory, *CurState.SeverityCategory);
+
+		if (bEscalate && ActiveState > WorstActiveState)
+		{
+			if (bDDoSAnalytics)
+			{
+				NotifySeverityEscalation.ExecuteIfBound(CurState.SeverityCategory);
+			}
+
+			WorstActiveState = ActiveState;
+		}
+	}
+}
+
+void FDDoSDetection::PreFrameReceive(float DeltaTime)
+{
+	if (bDDoSDetection)
+	{
+		StartFrameRecvTimestamp = FPlatformTime::Seconds();
+		bMetEscalationConditionsThisFrame = false;
+
+		if (HitchTimeQuotaMS > 0 && EndFrameRecvTimestamp != 0.0)
+		{
+			double HitchTimeMS = (StartFrameRecvTimestamp - EndFrameRecvTimestamp) * 1000.0;
+
+			if ((((int32)HitchTimeMS) - HitchTimeQuotaMS) > 0)
+			{
+				HitchFrameCount++;
+
+				UE_LOG(PacketHandlerLog, Verbose, TEXT("Detected '%i' successive hitches outside NetDriver Tick. Last Hitch: %fms (Max: %ims)"),
+						HitchFrameCount, HitchTimeMS, HitchTimeQuotaMS);
+			}
+			else
+			{
+				HitchFrameCount = 0;
+			}
+		}
+
+		// At the start of every frame, adjust the DDoS detection based upon DeltaTime - unless there is excessive hitching
+		FrameAdjustment = (HitchFrameCount > 0 && HitchFrameCount > HitchFrameTolerance) ? 1.f : (double)DeltaTime / ExpectedFrameTime;
+
+		if (ActiveState > 0 && CooloffTime > 0 && (float)(StartFrameRecvTimestamp - LastMetEscalationConditions) > (float)CooloffTime)
+		{
+			UpdateSeverity(false);
+		}
+
+		DetectionSeverity[ActiveState].ApplyAdjustedState(*this, FMath::Max(0.25f, FrameAdjustment));
+
+		if (((StartFrameRecvTimestamp - LastPerSecQuotaBegin) - 1.0) > 0.0)
+		{
+			UE_CLOG(DroppedPacketCounter > 0, PacketHandlerLog, Warning,
+				TEXT("DDoS Detection dropped '%i' packets during last second (bHitFrameNonConnLimit: %i, bHitFrameNetConnLimit: %i, ")
+				TEXT("DetectionSeverity: %s)."),
+				DroppedPacketCounter, (int32)bHitFrameNonConnLimit, (int32)bHitFrameNetConnLimit,
+				*DetectionSeverity[ActiveState].SeverityCategory);
+
+
+			// Record the last quota
+			check(CounterPerSecHistory.Num() > 0);
+
+			LastCounterPerSecHistoryIdx++;
+			LastCounterPerSecHistoryIdx = (LastCounterPerSecHistoryIdx >= CounterPerSecHistory.Num()) ? 0 : LastCounterPerSecHistoryIdx;
+
+			CounterPerSecHistory[LastCounterPerSecHistoryIdx] = *this;
+
+
+			LastPerSecQuotaBegin = StartFrameRecvTimestamp;
+			NonConnPacketCounter = 0;
+			NetConnPacketCounter = 0;
+			BadPacketCounter = 0;
+			ErrorPacketCounter = 0;
+			DroppedPacketCounter = 0;
+			WorstFrameReceiveTimeMS = 0;
+		}
+
+		StartFramePacketCount = NonConnPacketCounter;
+
+		if (LogHitCounter >= DDoSLogSpamLimit)
+		{
+			UE_LOG(PacketHandlerLog, Warning, TEXT("Previous frame hit DDoS LogHitCounter limit - hit count: %i (Max: %i)"), LogHitCounter,
+					DDoSLogSpamLimit);
+		}
+
+		LogHitCounter = 0;
+		bHitFrameNonConnLimit = false;
+		bHitFrameNetConnLimit = false;
+	}
+}
+
+void FDDoSDetection::PostFrameReceive()
+{
+	if (bDDoSDetection)
+	{
+		// Some packet counters require an end-frame check for DDoS detection
+		CheckNonConnQuotasAndLimits();
+
+
+		EndFrameRecvTimestamp = FPlatformTime::Seconds();
+
+		int32 FrameReceiveTimeMS = (int32)((EndFrameRecvTimestamp - StartFrameRecvTimestamp) * 1000.0);
+
+		WorstFrameReceiveTimeMS = FMath::Max(FrameReceiveTimeMS, WorstFrameReceiveTimeMS);
+	}
+}
+
+bool FDDoSDetection::CheckNonConnQuotasAndLimits()
+{
+	bool bReturnVal = false;
+	double CurTime = FPlatformTime::Seconds();
+	int32 TimePassedMS = (int32)((CurTime - StartFrameRecvTimestamp) * 1000.0);
+
+	if (HasHitQuota(*this, TimePassedMS))
+	{
+		UpdateSeverity(true);
+	}
+	// Check if we're still at the conditions which led to the current escalated state
+	else if (!bMetEscalationConditionsThisFrame && ActiveState > 0)
+	{
+		const int32 PrevState = ActiveState - 1;
+
+		if (DetectionSeverity[PrevState].HasHitQuota(*this, TimePassedMS))
+		{
+			LastMetEscalationConditions = CurTime;
+			bMetEscalationConditionsThisFrame = true;
+		}
+	}
+
+
+	// NOTE: PacketLimitPerFrame == 0 is a valid value, and blocks all non-NetConnection packets
+	bReturnVal = PacketLimitPerFrame == 0 || (PacketLimitPerFrame > 0 && (NonConnPacketCounter - StartFramePacketCount) >= PacketLimitPerFrame);
+	bReturnVal = bReturnVal || (PacketTimeLimitMSPerFrame > 0 && TimePassedMS > PacketTimeLimitMSPerFrame);
+
+	return bReturnVal;
+}
diff --git a/Engine/Source/Runtime/Net/Core/Private/Net/Core/Misc/PacketAudit.cpp b/Engine/Source/Runtime/Net/Core/Private/Net/Core/Misc/PacketAudit.cpp
new file mode 100644
--- /dev/null
+++ b/Engine/Source/Runtime/Net/Core/Private/Net/Core/Misc/PacketAudit.cpp
@@ -0,0 +1,515 @@
+// Copyright 1998-2019 Epic Games, Inc. All Rights Reserved.
+
+/**
+ * Includes
+ */
+#include "PacketAudit.h"
+#include "PacketHandler.h"
+
+#include "Misc/CommandLine.h"
+#include "Serialization/BitReader.h"
+#include "Serialization/MemoryArchive.h"
+
+
+#if !UE_BUILD_SHIPPING
+/**
+ * Defines
+ */
+
+/** The general name used/adapted for mutexes and shared memory */
+#define AUDIT_MUTEX_NAME	TEXT("UE4PacketAudit")
+
+/** Hardcoded amount of shared memory reserved for both the send/receive maps */
+#define AUDIT_MAPPING_SIZE	(1024 * 1024 * 32)
+
+/** Maximum allowed number of packets in the send/receive maps, to detect leaks */
+#define AUDIT_MAX_PACKETS	2048
+
+
+/**
+ * Globals
+ */
+FPacketAudit* GPacketAuditor = nullptr;
+
+
+/**
+ * Stored data for a single packet stage
+ */
+struct FPacketStageData
+{
+	/** The size of the packet, in bits */
+	uint32 SizeBits;
+
+	/** The CRC of this packet stage */
+	uint32 StageCRC;
+
+
+	friend FArchive& operator << (FArchive& Ar, FPacketStageData& StageData)
+	{
+		return Ar << StageData.SizeBits << StageData.StageCRC;
+	}
+
+	FString ToString() const
+	{
+		return FString::Printf(TEXT("(SizeBits: %i, StageCRC: %08X)"), SizeBits, StageCRC);
+	}
+};
+
+/**
+ * Stored data for an entire packet
+ */
+struct FPacketAuditData
+{
+	/** The size of the packet, in bits */
+	uint32								SizeBits;
+
+	/** Map of the stage name, and the packet data for this stage */
+	TMap<FString, FPacketStageData>		StageMap;
+
+
+	friend FArchive& operator << (FArchive& Ar, FPacketAuditData& AuditData)
+	{
+		return Ar << AuditData.SizeBits << AuditData.StageMap;
+	}
+
+	FString ToString(int32 BaseIndent) const
+	{
+		FString ReturnVal = TEXT("");
+		FString IndentStr = TEXT("");
+
+		for (int32 i=0; i<BaseIndent; i++)
+		{
+			IndentStr +=  TEXT("\t");
+		}
+
+		ReturnVal += IndentStr + TEXT("(\r\n");
+		ReturnVal += IndentStr + FString::Printf(TEXT("\tSizeBits: %i,\r\n"), SizeBits);
+		ReturnVal += IndentStr + TEXT("\tStageMap:\r\n");
+		ReturnVal += IndentStr + TEXT("\t(\r\n");
+
+		for (TMap<FString, FPacketStageData>::TConstIterator It(StageMap); It; ++It)
+		{
+			ReturnVal += IndentStr + FString::Printf(TEXT("\t\t(StageName: %s, StageData: %s),\r\n"), *It.Key(),
+														*It.Value().ToString());
+		}
+
+		ReturnVal += IndentStr + TEXT("\t)\r\n");
+		ReturnVal += IndentStr + TEXT(")");
+
+		return ReturnVal;
+	}
+};
+
+/**
+ * Tracks a pending packet send
+ */
+struct FPendingPacket
+{
+	/** The FBitWriter associated with the packet */
+	FBitWriter*							Writer;
+
+	/** The stage data being collected for the packet */
+	TMap<FString, FPacketStageData>		StageMap;
+};
+
+
+/** Map of low-level packet CRC's, and the audit data associated with those packets, for packets being sent. */
+static TMap<uint32, FPacketAuditData>		GSendPacketAudits;
+
+/** As above, except for received packets */
+static TMap<uint32, FPacketAuditData>		GReceivePacketAudits;
+
+/** Collects stage data, for the current packet being sent (local only) */
+static FPendingPacket						GPendingSendPacket;
+
+/** Stores the CRC for the current packet being received */
+static uint32								GCurrentReceivePacket;
+
+/** Stores the CRC for the packet that was last removed from GReceivePacketAudits (reset to 0 when receiving a new packet) */
+static uint32								GLastReceivePacket = 0;
+
+
+/**
+ * FSharedMemoryReader - based on FMemoryReader (just takes a pointer instead of array)
+ *
+ * NOTE: Hardcoded to AUDIT_MAPPING_SIZE
+ */
+class FSharedMemoryReader : public FMemoryArchive
+{
+public:
+	/**
+	 * Default constructor
+	 */
+	FSharedMemoryReader(uint8* InData)
+	{
+		MemSizeData = (uint32*)InData;
+		MemPayloadData = InData + sizeof(uint32);
+
+		check(*MemSizeData <= (AUDIT_MAPPING_SIZE - sizeof(uint32)));
+
+		this->SetIsLoading(true);
+	}
+
+	virtual int64 TotalSize() override
+	{
+		return *MemSizeData;
+	}
+
+	virtual void Serialize(void* Data, int64 Num) override
+	{
+		if (Num > 0 && !ArIsError)
+		{
+			if (Offset + Num <= TotalSize())
+			{
+				FMemory::Memcpy(Data, MemPayloadData + Offset, Num);
+				Offset += Num;
+			}
+			else
+			{
+				ArIsError = true;
+			}
+		}
+	}
+
+private:
+	/** The memory location storing the size of the archive */
+	uint32*		MemSizeData;
+
+	/** The memory location storing the archive payload */
+	uint8*		MemPayloadData;
+};
+
+/**
+ * FSharedMemoryWriter	- as above, same as FMemoryWriter except takes a pointer
+ *
+ * NOTE: Hardcoded to AUDIT_MAPPING_SIZE
+ * NOTE: Immediately invalidates existing data upon construction
+ */
+class FSharedMemoryWriter : public FMemoryArchive
+{
+public:
+	/**
+	 * Default constructor
+	 */
+	FSharedMemoryWriter(uint8* InData)
+	{
+		MemSizeData = (uint32*)InData;
+		MemPayloadData = InData + sizeof(uint32);
+
+		this->SetIsSaving(true);
+
+		*MemSizeData = 0;
+	}
+
+	virtual int64 TotalSize() override
+	{
+		return *MemSizeData;
+	}
+
+	virtual void Serialize(void* Data, int64 Num) override
+	{
+		check(Offset + Num < (AUDIT_MAPPING_SIZE - sizeof(uint32)));
+
+		if (Num > 0)
+		{
+			FMemory::Memcpy(MemPayloadData + Offset, Data, Num);
+
+			Offset += Num;
+
+			*MemSizeData = Offset;
+		}
+	}
+
+private:
+	/** The memory location storing the size of the archive */
+	uint32*		MemSizeData;
+
+	/** The memory location storing the archive payload */
+	uint8*		MemPayloadData;
+};
+#endif
+
+
+/**
+ * FPacketAudit
+ */
+
+FPacketAudit::FPacketAudit(FPlatformProcess::FSemaphore* InGameMutex)
+	: GameMutex(InGameMutex)
+{
+#if !UE_BUILD_SHIPPING
+	SharedMutex = FPlatformProcess::NewInterprocessSynchObject(AUDIT_MUTEX_NAME, true);
+
+	FString SendMemName = FString(AUDIT_MUTEX_NAME) + (GIsServer ? TEXT("Server") : TEXT("Client")) + TEXT("Send");
+	FString ReceiveMemName = FString(AUDIT_MUTEX_NAME) + (GIsServer ? TEXT("Client") : TEXT("Server")) + TEXT("Send");
+
+	SendSharedMemory = FPlatformMemory::MapNamedSharedMemoryRegion(SendMemName, true,
+					(FPlatformMemory::ESharedMemoryAccess::Read | FPlatformMemory::ESharedMemoryAccess::Write), AUDIT_MAPPING_SIZE);
+
+	ReceiveSharedMemory = FPlatformMemory::MapNamedSharedMemoryRegion(ReceiveMemName, true,
+					(FPlatformMemory::ESharedMemoryAccess::Read | FPlatformMemory::ESharedMemoryAccess::Write), AUDIT_MAPPING_SIZE);
+
+
+	// @todo #JohnB: There is probably a way to fix the shared memory platform code, so this works without admin. Low priority though.
+	if (SendSharedMemory == nullptr || ReceiveSharedMemory == nullptr)
+	{
+		LowLevelFatalError(TEXT("You need to run UE4 as administrator, for packet auditing to work."));
+	}
+
+	FSharedMemoryWriter WipeSendSharedMemory((uint8*)SendSharedMemory->GetAddress());
+	FSharedMemoryWriter WipeReceiveSharedMemory((uint8*)ReceiveSharedMemory->GetAddress());
+#endif
+}
+
+FPacketAudit::~FPacketAudit()
+{
+#if !UE_BUILD_SHIPPING
+	check(GameMutex != nullptr);
+	check(SharedMutex != nullptr);
+	check(SendSharedMemory != nullptr);
+	check(ReceiveSharedMemory != nullptr);
+
+	GameMutex->Unlock();
+
+	FPlatformProcess::DeleteInterprocessSynchObject(GameMutex);
+	FPlatformProcess::DeleteInterprocessSynchObject(SharedMutex);
+
+	GameMutex = nullptr;
+	SharedMutex = nullptr;
+
+	FPlatformMemory::UnmapNamedSharedMemoryRegion(SendSharedMemory);
+	FPlatformMemory::UnmapNamedSharedMemoryRegion(ReceiveSharedMemory);
+#endif
+}
+
+void FPacketAudit::Init()
+{
+#if !UE_BUILD_SHIPPING
+	if (GPacketAuditor == nullptr && FParse::Param(FCommandLine::Get(), TEXT("PacketAudit")))
+	{
+		FString MutexName = FString(AUDIT_MUTEX_NAME) + (GIsServer ? TEXT("ServerProcess") : TEXT("ClientProcess"));
+		FPlatformProcess::FSemaphore* CurGameMutex = FPlatformProcess::NewInterprocessSynchObject(MutexName, true);
+
+		if (CurGameMutex->TryLock(1))
+		{
+			GPacketAuditor = new FPacketAudit(CurGameMutex);
+		}
+		else
+		{
+			UE_LOG(PacketHandlerLog, Log, TEXT("%s"), *(FString(TEXT("Packet auditor already active for a game ")) +
+					(GIsServer ? TEXT("server") : TEXT("client")) + TEXT(", can't start multiple instances.")));
+
+			FPlatformProcess::DeleteInterprocessSynchObject(CurGameMutex);
+		}
+	}
+#endif
+}
+
+void FPacketAudit::AddStage_Internal(FString StageName, FBitWriter& OutPacket, bool bByteAligned)
+{
+#if !UE_BUILD_SHIPPING
+	if (GPendingSendPacket.Writer == nullptr)
+	{
+		GPendingSendPacket.Writer = &OutPacket;
+	}
+
+	check(GPendingSendPacket.Writer == &OutPacket);
+
+	FPacketStageData& StageData = GPendingSendPacket.StageMap.Add(StageName);
+
+	StageData.SizeBits = (bByteAligned ? (OutPacket.GetNumBytes() * 8) : OutPacket.GetNumBits());
+	StageData.StageCRC = PacketCRC(OutPacket.GetData(), StageData.SizeBits);
+#endif
+}
+
+void FPacketAudit::CheckStage_Internal(FString StageName, FBitReader& InPacket, bool bByteAligned)
+{
+#if !UE_BUILD_SHIPPING
+	{
+		FScopedAuditAccess AuditScope(*this);
+		FPacketAuditData* AuditData = GReceivePacketAudits.Find(GCurrentReceivePacket);
+
+		if (AuditData != nullptr)
+		{
+			FPacketStageData* StageData = AuditData->StageMap.Find(StageName);
+
+			if (StageData != nullptr)
+			{
+				uint32 BitsLeft = (bByteAligned ? ((InPacket.GetNumBytes() * 8) - InPacket.GetPosBits()) : InPacket.GetBitsLeft());
+
+				if (BitsLeft == StageData->SizeBits)
+				{
+					uint32 PacketStageCRC = PacketCRC(InPacket.GetData(), BitsLeft);
+
+					if (PacketStageCRC == StageData->StageCRC)
+					{
+						AuditData->StageMap.Remove(StageName);
+
+						if (AuditData->StageMap.Num() == 0)
+						{
+							GReceivePacketAudits.Remove(GCurrentReceivePacket);
+
+							GLastReceivePacket = GCurrentReceivePacket;
+						}
+					}
+					else
+					{
+						DumpAuditData(AuditScope, GCurrentReceivePacket);
+						LowLevelFatalError(TEXT("%s: Expected Stage CRC '%08X', got Stage CRC '%08X'."), *StageName,
+											StageData->StageCRC, PacketStageCRC);
+					}
+				}
+				else
+				{
+					DumpAuditData(AuditScope, GCurrentReceivePacket);
+					LowLevelFatalError(TEXT("%s: Expected '%i' bits left, got '%i'."), *StageName, StageData->SizeBits, BitsLeft);
+				}
+			}
+			else
+			{
+				DumpAuditData(AuditScope, GCurrentReceivePacket);
+				LowLevelFatalError(TEXT("%s: Could not find StageName."), *StageName);
+			}
+		}
+		else
+		{
+			DumpAuditData(AuditScope, GCurrentReceivePacket);
+
+			if (GLastReceivePacket == GCurrentReceivePacket)
+			{
+				LowLevelFatalError(TEXT("%s: Already finished proccesing all stages from packet."), *StageName);
+			}
+			else
+			{
+				LowLevelFatalError(TEXT("%s: Failed to find packet CRC: %08X. Duplicate/resent packet?"), *StageName,
+									GCurrentReceivePacket);
+			}
+		}
+	}
+#endif
+}
+
+void FPacketAudit::NotifyLowLevelSend_Internal(FBitWriter& OutPacket)
+{
+#if !UE_BUILD_SHIPPING
+	if (GPendingSendPacket.Writer != nullptr)
+	{
+		// @todo #JohnB: Restore when you have the whole packet pipeline unified
+		//check(GPendingSendPacket.Writer == &OutPacket);
+
+
+		uint32 OutSizeBits = OutPacket.GetNumBits();
+
+		if (OutSizeBits > 0)
+		{
+			uint32 OutPacketCRC = PacketCRC(OutPacket.GetData(), OutSizeBits);
+
+			{
+				FScopedAuditAccess AuditScope(*this);
+
+				// @todo: Most often this will be due to reliability handling. Collisions may be possible though.
+				check(!GSendPacketAudits.Contains(OutPacketCRC));
+
+
+				check(GSendPacketAudits.Num() < AUDIT_MAX_PACKETS);
+
+				FPacketAuditData& AuditData = GSendPacketAudits.Add(OutPacketCRC);
+
+				AuditData.SizeBits = OutSizeBits;
+
+				AuditData.StageMap = GPendingSendPacket.StageMap;
+			}
+		}
+
+		GPendingSendPacket.Writer = nullptr;
+		GPendingSendPacket.StageMap.Empty();
+	}
+#endif
+}
+
+void FPacketAudit::NotifyLowLevelReceive_Internal(FBitReader& InPacket)
+{
+#if !UE_BUILD_SHIPPING
+	GCurrentReceivePacket = PacketCRC(InPacket.GetData(), InPacket.GetNumBits());
+	GLastReceivePacket = 0;
+#endif
+}
+
+uint32 FPacketAudit::PacketCRC(uint8* Data, uint32 DataLenBits)
+{
+	uint32 ReturnVal = 0;
+
+	if (DataLenBits > 0)
+	{
+		uint32 PacketByteSize = (DataLenBits + 7) >> 3;
+		uint8* PacketData = new uint8[PacketByteSize];
+
+		PacketData[PacketByteSize-1] = 0;
+
+		appBitsCpy(PacketData, 0, Data, 0, DataLenBits);
+
+		ReturnVal = FCrc::MemCrc32(PacketData, PacketByteSize);
+
+		delete[] PacketData;
+	}
+
+	return ReturnVal;
+}
+
+#if !UE_BUILD_SHIPPING
+void FPacketAudit::DumpAuditData(FScopedAuditAccess& AuditLock, uint32 InPacketCRC/*=0*/)
+{
+	FString FinalLogStr = (InPacketCRC != 0 ? FString::Printf(TEXT("PacketCRC: %08X,\r\n"), InPacketCRC) : TEXT(""));
+	
+	FinalLogStr += TEXT("GSendPacketAudits:\r\n(\r\n");
+
+	for (TMap<uint32, FPacketAuditData>::TConstIterator It(GSendPacketAudits); It; ++It)
+	{
+		FinalLogStr += FString::Printf(TEXT("\tPacketCRC: %08X,\r\n"), It.Key());
+		FinalLogStr += TEXT("\tAuditData:\r\n");
+		FinalLogStr += It.Value().ToString(1) + TEXT("\r\n");
+	}
+
+	FinalLogStr += TEXT("),\r\nGReceivePacketAudits:\r\n(\r\n");
+
+	for (TMap<uint32, FPacketAuditData>::TConstIterator It(GReceivePacketAudits); It; ++It)
+	{
+		FinalLogStr += FString::Printf(TEXT("\tPacketCRC: %08X,\r\n"), It.Key());
+		FinalLogStr += TEXT("\tAuditData:\r\n");
+		FinalLogStr += It.Value().ToString(1) + TEXT("\r\n");
+	}
+
+	FinalLogStr += TEXT(")");
+
+	UE_LOG(PacketHandlerLog, Log, TEXT("%s"), *FinalLogStr);
+}
+
+
+FPacketAudit::FScopedAuditAccess::FScopedAuditAccess(FPacketAudit& InAuditor)
+	: Auditor(InAuditor)
+{
+	Auditor.SharedMutex->Lock();
+
+	{
+		FSharedMemoryReader SendReader((uint8*)Auditor.SendSharedMemory->GetAddress());
+		FSharedMemoryReader ReceiveReader((uint8*)Auditor.ReceiveSharedMemory->GetAddress());
+
+		SendReader << GSendPacketAudits;
+		ReceiveReader << GReceivePacketAudits;
+	}
+}
+
+FPacketAudit::FScopedAuditAccess::~FScopedAuditAccess()
+{
+	{
+		FSharedMemoryWriter SendWriter((uint8*)Auditor.SendSharedMemory->GetAddress());
+		FSharedMemoryWriter ReceiveWriter((uint8*)Auditor.ReceiveSharedMemory->GetAddress());
+
+		SendWriter << GSendPacketAudits;
+		ReceiveWriter << GReceivePacketAudits;
+	}
+
+	Auditor.SharedMutex->Unlock();
+}
+#endif
+
diff --git a/Engine/Source/Runtime/Net/Core/Public/Net/Core/Analytics/NetAnalytics.h b/Engine/Source/Runtime/Net/Core/Public/Net/Core/Analytics/NetAnalytics.h
new file mode 100644
--- /dev/null
+++ b/Engine/Source/Runtime/Net/Core/Public/Net/Core/Analytics/NetAnalytics.h
@@ -0,0 +1,389 @@
+// Copyright 1998-2019 Epic Games, Inc. All Rights Reserved.
+
+#pragma once
+
+/**
+ * Net Analytics Aggregation
+ *
+ * Net analytics aggregation is used by the NetDriver (and primarily Oodle), to collect analytics during the lifetime of the NetDriver,
+ * and to aggregate it so that it can be dispatched in one consolidated analytics event, at the end of the NetDriver's lifetime.
+ *
+ * This is useful for e.g. per-NetConnection analytics data, which we want to acquire but don't want to spam the analytics service with,
+ * so we need an interface to aggregate at the NetDriver level (in principle, this should be reusable outside of the netcode too).
+ *
+ *
+ * Usage:
+ * To collect analytics data you will need to subclass FNetAnalyticsData and use the REGISTER_NET_ANALYTICS macro,
+ * to create and register an instance of your data holder, with a NetDriver's NetAnalyticsAggregator,
+ * using a unique name that only your code uses (preferably the analytics event name).
+ *
+ * This will also need to be registered in *Engine.ini, for each NetDriverName, e.g.:
+ *	[GameNetDriver NetAnalyticsAggregatorConfig]
+ *	+NetAnalyticsData=(DataName="Core.ServerNetConn", bEnabled=true)
+ *
+ * Then assign the result of REGISTER_NET_ANALYTICS, to store a reference to your data holder, in the location you collect analytics.
+ *
+ * The way your FNetAnalyticsData subclass collects and stores/aggregates data is up to you (and so is the thread-safety for that data),
+ * just implement SendAnalytics to process and dispatch all of the aggregated data upon NetDriver shutdown. It's only called once.
+ *
+ *
+ * Multithreaded lifetime:
+ * All multithreaded analytics aggregation objects/data/events must be complete by NetDriver shut down.
+ * Staying within the normal course of netcode send/receive paths should achieve this.
+ */
+
+// Includes
+#include "CoreMinimal.h"
+#include "Templates/Atomic.h"
+
+
+// Forward Declarations
+class IAnalyticsProvider;
+class FNetAnalyticsAggregator;
+
+
+// Defines
+
+/** Whether or not to enable multithreaded support (future proofing for netcode multithreading) - comes with a performance cost */
+#define NET_ANALYTICS_MULTITHREADING 0
+
+
+#if NET_ANALYTICS_MULTITHREADING
+	constexpr const ESPMode NetAnalyticsThreadSafety = ESPMode::ThreadSafe;
+#else
+	constexpr const ESPMode NetAnalyticsThreadSafety = ESPMode::Fast;
+#endif
+
+
+// Forward Declarations
+struct FNetAnalyticsData;
+
+
+// Typedefs
+template<class T=FNetAnalyticsData> using TNetAnalyticsDataRef = TSharedRef<T, NetAnalyticsThreadSafety>;
+template<class T=FNetAnalyticsData> using TNetAnalyticsDataPtr = TSharedPtr<T, NetAnalyticsThreadSafety>;
+template<class T=FNetAnalyticsData> using TNetAnalyticsDataWeakPtr = TWeakPtr<T, NetAnalyticsThreadSafety>;
+
+
+// Defines
+
+/**
+ * Registers a named FNetAnalyticsData instance with the net analytics aggregator.
+ * Implemented this way, to support runtime name/type checks, without relying on statics.
+ *
+ * NOTE: May return nullptr! The analytics aggregator can selectively enable/disable analytics, based on the analytics data name
+ *
+ * @param Aggregator			The net analytics aggregator to register with
+ * @param AnalyticsDataType		The net analytics data type to create and register
+ * @param InDataName			The name associated with this analytics data - for looking it up
+ * @return						Returns a shared pointer to the newly created net analytics data instance
+ */
+#define REGISTER_NET_ANALYTICS(Aggregator, AnalyticsDataType, InDataName) \
+	StaticCastSharedPtr<AnalyticsDataType, FNetAnalyticsData, NetAnalyticsThreadSafety>( \
+		(Aggregator->RegisterAnalyticsData_Internal(MakeShareable(new AnalyticsDataType(), FNetAnalyticsDataDeleter()), InDataName, \
+			TEXT(#AnalyticsDataType))))
+
+
+// Globals
+
+/** Counter used to detect when a new FNetAnalyticsAggregator (i.e. new NetDriver) has been created - to trigger cleanup of TLS data */
+extern PACKETHANDLER_API TAtomic<uint8> GNetAnalyticsCounter;
+
+
+/**
+ * For use with TThreadedNetAnalyticsData - this struct is subclassed, and used to define/hold the analytics variables,
+ * that will be tracked/aggregated.
+ *
+ * This base type is required, due to how TThreadedNetAnalyticsData handles tracking of local analytics variables.
+ */
+struct FLocalNetAnalyticsStruct : public FVirtualDestructor
+{
+};
+
+
+/**
+ * Subclassed struct which holds net analytics data which is aggregated or to-be-aggregated, and dispatched upon SendAnalytics
+ */
+struct PACKETHANDLER_API FNetAnalyticsData : public FVirtualDestructor, public TSharedFromThis<FNetAnalyticsData, NetAnalyticsThreadSafety>
+{
+	friend class FNetAnalyticsAggregator;
+	friend struct FNetAnalyticsDataDeleter;
+
+protected:
+	/** The parent aggregator responsible for this data instance */
+	FNetAnalyticsAggregator* Aggregator;
+
+
+public:
+	/**
+	 * Default constructor
+	 */
+	FNetAnalyticsData()
+		: Aggregator(nullptr)
+	{
+	}
+
+	/**
+	 * Dispatches the aggregated analytics data - no further data should be added
+	 */
+	virtual void SendAnalytics() = 0;
+
+
+protected:
+	/**
+	 * Called by the Net Analytics Aggregator, when ready to send analytics - some implementations don't send immediately
+	 */
+	virtual void InternalSendAnalytics();
+
+	/**
+	 * Called just before the last shared reference to this data is released - used to trigger analytics send in threadsafe version
+	 */
+	virtual void NotifyFinalRelease()
+	{
+	}
+};
+
+/**
+ * Basic single-threaded-only (i.e. good for NetConnection level code) analytics data holder, which just wraps around a simple struct,
+ * which defines the analytics variables and implements their aggregation.
+ *
+ * The passed in struct will have to implement a CommitAnalytics function, which implements aggregation of the struct data,
+ * and the class which references this net analytics data will have to trigger this classes CommitAnalytics on the locally stored struct,
+ * when finished collecting local analytics data.
+ */
+template<class TDataStruct> struct TBasicNetAnalyticsData : public FNetAnalyticsData, protected TDataStruct
+{
+	GENERATE_MEMBER_FUNCTION_CHECK(CommitAnalytics, void,, TDataStruct&);
+
+	static_assert(THasMemberFunction_CommitAnalytics<TDataStruct>::Value, "TDataStruct must implement void CommitAnalytics(TDataStruct&)");
+
+	/**
+	 * Called by the class/code which is locally collecting analytics data, to commit the local data for aggregation when done.
+	 *
+	 * @param AnalyticsVars		The local analytics variables to be committed/aggregated
+	 */
+	void CommitAnalytics(TDataStruct& AnalyticsVars)
+	{
+		AnalyticsVars.CommitAnalytics(*this);
+	}
+};
+
+#if NET_ANALYTICS_MULTITHREADING
+/**
+ * Special subclass of FNetAnalyticsData that is designed for thread-safety and infrequent access (e.g. at NetConnection Close only).
+ * This is to future-proof for netcode multithreading, where PacketHandler level code is expected to run outside of the Game Thread.
+ *
+ * The last thread to release this analytics data, triggers SendAnalytics - thread safety is partially provided by shared pointer atomics.
+ */
+struct PACKETHANDLER_API FThreadedNetAnalyticsData : public FNetAnalyticsData
+{
+	/**
+	 * Default constructor
+	 */
+	FThreadedNetAnalyticsData();
+
+
+protected:
+	virtual void InternalSendAnalytics() override;
+
+	virtual void NotifyFinalRelease() override;
+
+
+protected:
+	/** Whether or not a thread has signaled that analytics are ready to be sent */
+	TAtomic<bool> bReadyToSend;
+};
+
+/**
+ * Special subclass of FThreadedNetAnalyticsData, which implements multithreaded synchronization, using Thread Local Storage (TLS),
+ * providing a complete solution for aggregating net analytics data, by caching a local copy of TDataStruct in TLS (one per thread),
+ * and automatically aggregating all instances upon NetDriver Shutdown.
+ *
+ * The passed in struct will have to implement a CommitAnalytics function, which implements aggregation of the struct data.
+ *
+ * Every time you need to update analytics, use GetLocalData to access the variables - don't permanently store the return value.
+ *
+ * There is a performance cost both to looking up thread_local's, and to mapping TDataStruct for TThreadedNetAnalyticsData.
+ * This code must only be enabled in a multithreaded environment, due to the unnecessary performance cost in single-threaded code.
+ */
+template<class TDataStruct> struct TThreadedNetAnalyticsData : public FThreadedNetAnalyticsData, protected TDataStruct
+{
+	GENERATE_MEMBER_FUNCTION_CHECK(CommitAnalytics, void,, TDataStruct&);
+
+	static_assert(THasMemberFunction_CommitAnalytics<TDataStruct>::Value, "TDataStruct must implement void CommitAnalytics(TDataStruct&)");
+	static_assert(TIsDerivedFrom<TDataStruct, FLocalNetAnalyticsStruct>::IsDerived, "TDataStruct must be derived from FLocalNetAnalyticsStruct");
+
+
+	/**
+	 * Default constructor
+	 */
+	TThreadedNetAnalyticsData()
+		: ThreadLocalData()
+	{
+	}
+
+	/**
+	 * Returns the current threads TDataStruct instance, for this net analytics data - creating/registering it, if necessary.
+	 *
+	 * @return	The current threads TDataStruct instance
+	 */
+	TDataStruct* GetLocalData()
+	{
+		/** Maps this threads net analytics variables, to the net analytics data handler which owns them - for fast lookup */
+		thread_local TMap<FThreadedNetAnalyticsData*, FLocalNetAnalyticsStruct*> LocalNetAnalyticsMap;
+
+		FLocalNetAnalyticsStruct** FoundVal = LocalNetAnalyticsMap.Find(this);
+
+		return (FoundVal != nullptr ? static_cast<TDataStruct*>(*FoundVal) : AddLocalData(LocalNetAnalyticsMap));
+	}
+
+
+protected:
+	/**
+	 * Handles creation/registration of a new TDataStruct instance - as well as occasional cleanup of stale TLS data.
+	 *
+	 * @param LocalNetAnalyticsMap	Reference to the TLS map linking local net analytics vars, to a (potentially stale) net analytics data pointer
+	 * @return						The newly created TDataStruct instance
+	 */
+	TDataStruct* AddLocalData(TMap<FThreadedNetAnalyticsData*, FLocalNetAnalyticsStruct*>& LocalNetAnalyticsMap)
+	{
+		TDataStruct* ReturnVal = new TDataStruct();
+
+		/** Maps a weak net analytics data pointer, to its own raw pointer - to cleanup stale entries in LocalNetAnalyticsMap */
+		thread_local TMap<TNetAnalyticsDataWeakPtr<FThreadedNetAnalyticsData>, FThreadedNetAnalyticsData*> StaleNetAnalyticsTracking;
+
+		LocalNetAnalyticsMap.Add(this, ReturnVal);
+		StaleNetAnalyticsTracking.Add(StaticCastSharedRef<FThreadedNetAnalyticsData, FNetAnalyticsData, NetAnalyticsThreadSafety>(AsShared()), this);
+		ThreadLocalData.Enqueue(ReturnVal);
+
+
+		/** Each thread tracks a counter signaling a new FNetAnalyticsAggregator/NetDriver instance, using this to trigger cleanup */
+		thread_local uint8 LastNetAnalyticsCounter = 0;
+
+		uint8 CurNetAnalyticsCounter = GNetAnalyticsCounter;
+
+		if (CurNetAnalyticsCounter != LastNetAnalyticsCounter)
+		{
+			for (auto It=StaleNetAnalyticsTracking.CreateIterator(); It; ++It)
+			{
+				if (!It.Key().IsValid())
+				{
+					LocalNetAnalyticsMap.Remove(It.Value());
+					It.RemoveCurrent();
+				}
+			}
+
+			LastNetAnalyticsCounter = CurNetAnalyticsCounter;
+		}
+
+		return ReturnVal;
+	}
+
+	virtual void NotifyFinalRelease() override
+	{
+		TDataStruct* CurAnalyticsData = nullptr;
+
+		while (ThreadLocalData.Dequeue(CurAnalyticsData))
+		{
+			CurAnalyticsData->CommitAnalytics(*this);
+
+			delete CurAnalyticsData;
+		}
+
+		FThreadedNetAnalyticsData::NotifyFinalRelease();
+	}
+
+
+protected:
+	/** Every thread creates an instance of TDataStruct locally, and queues the pointer here for later processing/deletion */
+	TQueue<TDataStruct*, EQueueMode::Mpsc> ThreadLocalData;
+};
+#endif
+
+/**
+ * Custom deleter for FNetAnalyticsData shared pointers
+ */
+struct FNetAnalyticsDataDeleter
+{
+	FORCEINLINE void operator()(FNetAnalyticsData* Data)
+	{
+		if (Data->Aggregator != nullptr)
+		{
+			Data->NotifyFinalRelease();
+		}
+
+		delete Data;
+	}
+};
+
+
+/**
+ * Central object (usually within NetDriver) which handles registration/retrieval/type-checking of net analytics data holders.
+ */
+class PACKETHANDLER_API FNetAnalyticsAggregator
+{
+public:
+	/**
+	 * Base constructor
+	 */
+	FNetAnalyticsAggregator(TSharedPtr<IAnalyticsProvider> InProvider, FName InNetDriverName);
+
+	FNetAnalyticsAggregator() = delete;
+
+	/**
+	 * Initialize the net analytics aggregator
+	 */
+	void Init();
+
+	/**
+	 * Initialize the net analytics aggregator config - must support hotfixing
+	 */
+	void InitConfig();
+
+	/**
+	 * Tells the analytics data holders to finish aggregating their analytics data, and to dispatch it.
+	 * Only called once, at NetDriver shutdown.
+	 */
+	void SendAnalytics();
+
+
+	/**
+	 * Use REGISTER_NET_ANALYTICS instead. Internal function, which registers a net analytics data holder, with an associated key name,
+	 * and does type checking to ensure there have been no mixups with the data holder type.
+	 *
+	 * NOTE: May return nullptr! Analytics can be selectively enabled/disabled, based on the specified analytics data name.
+	 *
+	 * @param InData		The analytics data shared reference to be registered
+	 * @param InDataName	The name given for referencing the analytics data
+	 * @param InTypeName	Compile-time derived type name for the analytics data - for type checking
+	 * @return				Returns the newly added analytics data shared reference
+	 */
+	TNetAnalyticsDataPtr<> RegisterAnalyticsData_Internal(TNetAnalyticsDataRef<> InData, const FName& InDataName, FString InTypeName);
+
+
+	/** Accessor for AnalyticsProvider */
+	const TSharedPtr<IAnalyticsProvider>& GetAnalyticsProvider()
+	{
+		return AnalyticsProvider;
+	}
+
+
+private:
+	/** The analytics provider we are aggregating data for */
+	TSharedPtr<IAnalyticsProvider> AnalyticsProvider;
+
+	/** The name of the NetDriver which owns this analytics aggregator - for retrieving NetDriver-specific config values */
+	FName NetDriverName;
+
+	/** Maps net analytics data holders, to their specified name */
+	TMap<FName, TNetAnalyticsDataRef<>> AnalyticsDataMap;
+
+	/** Maps analytics data holder names, to their type name - to verify types and prevent miscasting */
+	TMap<FName, FString> AnalyticsDataTypeMap;
+
+	/** Maps analytics data holder names, to a config value specifying whether that data holder is enabled or not */
+	TMap<FName, bool> AnalyticsDataConfigMap;
+
+	/** Whether or not analytics was already sent */
+	bool bSentAnalytics;
+};
diff --git a/Engine/Source/Runtime/Net/Core/Public/Net/Core/Misc/DDoSDetection.h b/Engine/Source/Runtime/Net/Core/Public/Net/Core/Misc/DDoSDetection.h
new file mode 100644
--- /dev/null
+++ b/Engine/Source/Runtime/Net/Core/Public/Net/Core/Misc/DDoSDetection.h
@@ -0,0 +1,358 @@
+// Copyright 1998-2019 Epic Games, Inc. All Rights Reserved.
+
+#pragma once
+
+// Includes
+#include "CoreMinimal.h"
+#include "HAL/PlatformTime.h"
+
+
+// Delegates
+
+/**
+ * Delegate for allowing analytics to receive notification of detected DDoS attacks
+ */
+DECLARE_DELEGATE_OneParam(FDDoSSeverityEscalation, FString /*SeverityCategory*/);
+
+
+
+/**
+ * Struct containing the per-second packet counters
+ */
+struct PACKETHANDLER_API FDDoSPacketCounters
+{
+	/** Counter for non-NetConnection packets received, since the last per second quota period began */
+	int32 NonConnPacketCounter;
+
+	/** Counter for NetConnection packets received, since the last per second quota period began */
+	int32 NetConnPacketCounter;
+
+	/** Counter for recently-disconnected NetConnection packets received, since the last per second quota period beegan */
+	int32 DisconnPacketCounter;
+
+	/** Counter for bad non-NetConnection packets received, since the last per second quota period began */
+	int32 BadPacketCounter;
+
+	/** Counter for non-DDoS packet errors received */
+	int32 ErrorPacketCounter;
+
+	/** Counter for the number of packets (of any type) dropped, since the last per second quota period began */
+	int32 DroppedPacketCounter;
+
+	/** The worst per-frame packet receive time in milliseconds, over the ~1 second this packet counter history represents */
+	int32 WorstFrameReceiveTimeMS;
+
+
+	FDDoSPacketCounters()
+		: NonConnPacketCounter(0)
+		, NetConnPacketCounter(0)
+		, DisconnPacketCounter(0)
+		, BadPacketCounter(0)
+		, ErrorPacketCounter(0)
+		, DroppedPacketCounter(0)
+		, WorstFrameReceiveTimeMS(0)
+	{
+	}
+};
+
+
+/**
+ * Stores the DDoS detection state (either settings from the config file, or the active DDoS detection state)
+ */
+struct PACKETHANDLER_API FDDoSState
+{
+	/** The number of packets/sec before the next stage of DDoS detection is triggered */
+	int32 EscalateQuotaPacketsPerSec;
+
+	/** The number of recently disconnected NetConnection packets/sec, before the next stage of DDoS detection is triggered. */
+	int32 EscalateQuotaDisconnPacketsPerSec;
+
+	/** The number of bad (failed to process correctly) packets/sec, before the next stage of DDoS detection is triggered */
+	int32 EscalateQuotaBadPacketsPerSec;
+
+	/** The amount of time spent processing packets, before the next stage of DDoS detection is triggered */
+	int16 EscalateTimeQuotaMSPerFrame;
+
+	/** The limit for the number of non-NetConnection packets to process, each frame */
+	int32 PacketLimitPerFrame;
+
+	/** The limit for time spent processing non-NetConnection packets, each frame (counts all packets time, non-NetConn and NetConn) */
+	int32 PacketTimeLimitMSPerFrame;
+
+	/** The limit for time spent processing NetConnection packets, each frame (counts all packets time, non-NetConn and NetConn) */
+	int32 NetConnPacketTimeLimitMSPerFrame;
+
+	/** The amount of time, in seconds, before the current DDoS severity category cools off and de-escalates */
+	int32 CooloffTime;
+
+
+	FDDoSState()
+		: EscalateQuotaPacketsPerSec(-1)
+		, EscalateQuotaDisconnPacketsPerSec(-1)
+		, EscalateQuotaBadPacketsPerSec(-1)
+		, EscalateTimeQuotaMSPerFrame(-1)
+		, PacketLimitPerFrame(-1)
+		, PacketTimeLimitMSPerFrame(-1)
+		, NetConnPacketTimeLimitMSPerFrame(-1)
+		, CooloffTime(-1)
+	{
+	}
+
+	/**
+	 * Whether or not the specified counters and time passed has hit any of the quota's
+	 *
+	 * @param InCounter		The counters to check against
+	 * @param TimePassedMS	The amount of time passed in milliseconds, since the beginning of this state
+	 * @return				Whether or not the quota was hit
+	 */
+	bool HasHitQuota(FDDoSPacketCounters& InCounter, int32 TimePassedMS) const
+	{
+		const bool bAtQuota = EscalateQuotaPacketsPerSec > 0 && InCounter.NonConnPacketCounter >= EscalateQuotaPacketsPerSec;
+		const bool bAtDisconnQuota = EscalateQuotaDisconnPacketsPerSec > 0 && InCounter.DisconnPacketCounter >= EscalateQuotaDisconnPacketsPerSec;
+		const bool bAtBadQuota = EscalateQuotaBadPacketsPerSec > 0 && InCounter.BadPacketCounter >= EscalateQuotaBadPacketsPerSec;
+		const bool bAtTimeQuota = EscalateTimeQuotaMSPerFrame > 0 && TimePassedMS > EscalateTimeQuotaMSPerFrame;
+
+		return bAtQuota || bAtDisconnQuota || bAtBadQuota || bAtTimeQuota;
+	}
+};
+
+/**
+ * DDoS detection state, with functions for applying the state to active DDoS detection
+ */
+struct PACKETHANDLER_API FDDoSStateConfig : public FDDoSState
+{
+	/** The name of the DDoS severity level this config section represents */
+	FString SeverityCategory;
+
+
+	FDDoSStateConfig()
+		: SeverityCategory(TEXT(""))
+	{
+	}
+
+	void ApplyState(FDDoSState& Target)
+	{
+		Target.EscalateQuotaPacketsPerSec			= EscalateQuotaPacketsPerSec;
+		Target.EscalateQuotaDisconnPacketsPerSec	= EscalateQuotaDisconnPacketsPerSec;
+		Target.EscalateQuotaBadPacketsPerSec		= EscalateQuotaBadPacketsPerSec;
+		Target.EscalateTimeQuotaMSPerFrame			= EscalateTimeQuotaMSPerFrame;
+		Target.PacketLimitPerFrame					= PacketLimitPerFrame;
+		Target.PacketTimeLimitMSPerFrame			= PacketTimeLimitMSPerFrame;
+		Target.NetConnPacketTimeLimitMSPerFrame		= NetConnPacketTimeLimitMSPerFrame;
+		Target.CooloffTime							= CooloffTime;
+	}
+
+	/**
+	 * Applies only the per-frame adjusted state (based on expected vs actual framerate), to active DDoS protection.
+	 * ApplyState should be called, first.
+	 */
+	void ApplyAdjustedState(FDDoSState& Target, float FrameAdjustment)
+	{
+		// Exclude escalation triggers from this
+		//Target.EscalateTimeQuotaMSPerFrame		= (EscalateTimeQuotaMSPerFrame == -1 ? -1 : EscalateTimeQuotaMSPerFrame * FrameAdjustment);
+		Target.PacketLimitPerFrame				= (PacketLimitPerFrame == -1 ? -1 : PacketLimitPerFrame * FrameAdjustment);
+		Target.PacketTimeLimitMSPerFrame		= (PacketTimeLimitMSPerFrame == -1 ? -1 : PacketTimeLimitMSPerFrame * FrameAdjustment);
+		Target.NetConnPacketTimeLimitMSPerFrame	= (NetConnPacketTimeLimitMSPerFrame == -1 ? -1 : NetConnPacketTimeLimitMSPerFrame * FrameAdjustment);
+	}
+};
+
+
+/**
+ * The main DDoS detection tracking class, for counting packets and applying restrictions.
+ * Implemented separate to the NetDriver, to allow wider use e.g. potentially at socket level, if useful.
+ */
+class PACKETHANDLER_API FDDoSDetection : protected FDDoSPacketCounters, protected FDDoSState
+{
+public:
+	/**
+	 * Default constructor
+	 */
+	FDDoSDetection();
+
+
+	/**
+	 * Initializes the DDoS detection settings
+	 *
+	 * @param MaxTickRate	The maximum tick rate of the server
+	 */
+	void Init(int32 MaxTickRate);
+
+	/**
+	 * Initializes the settings from the .ini file - must support reloading of settings on-the-fly
+	 */
+	void InitConfig();
+
+	/**
+	 * Updates the current DDoS detection severity state
+	 *
+	 * @param bEscalate		Whether or not we are escalating or de-escalating the severity state
+	 */
+	void UpdateSeverity(bool bEscalate);
+
+
+	/**
+	 * Triggered before packet receive begins, during the current frame
+	 */
+	void PreFrameReceive(float DeltaTime);
+
+	/**
+	 * Triggered after packet receive ends, during the current frame
+	 */
+	void PostFrameReceive();
+
+
+	/**
+	 * Rate limited call to CheckNonConnQuotasAndLimits
+	 */
+	void CondCheckNonConnQuotasAndLimits()
+	{
+		// Limit checks to once every 128 packets
+		if (((NonConnPacketCounter + DisconnPacketCounter) & 0x7F) == 0)
+		{
+			bHitFrameNonConnLimit = CheckNonConnQuotasAndLimits();
+		}
+	}
+
+	/**
+	 * Rate limited call to CheckNetConnLimits
+	 */
+	void CondCheckNetConnLimits()
+	{
+		// Limit checks to once every 128 packets
+		if ((NetConnPacketCounter & 0x7F) == 0)
+		{
+			bHitFrameNetConnLimit = CheckNetConnLimits();
+		}
+	}
+
+
+	/**
+	 * Accessor for bDDoSLogRestrictions - doubles as a per-frame logspam counter, automatically disabling logs after a quota
+	 */
+	bool CheckLogRestrictions()
+	{
+		return bDDoSLogRestrictions || (bDDoSDetection && ++LogHitCounter > DDoSLogSpamLimit);
+	}
+
+
+	// Brief accessors
+
+	bool IsDDoSDetectionEnabled() const		{ return bDDoSDetection; }
+	bool IsDDoSAnalyticsEnabled() const		{ return bDDoSAnalytics; }
+	bool ShouldBlockNonConnPackets() const	{ return bHitFrameNonConnLimit; }
+	bool ShouldBlockNetConnPackets() const	{ return bHitFrameNetConnLimit; }
+
+	void IncNonConnPacketCounter()			{ ++NonConnPacketCounter; }
+	int32 GetNonConnPacketCounter() const	{ return NonConnPacketCounter; }
+	void IncNetConnPacketCounter()			{ ++NetConnPacketCounter; }
+	int32 GetNetConnPacketCounter() const	{ return NetConnPacketCounter; }
+	void IncDisconnPacketCounter()			{ ++ DisconnPacketCounter; }
+	int32 GetDisconnPacketCounter() const	{ return DisconnPacketCounter; }
+	void IncBadPacketCounter()				{ ++BadPacketCounter; }
+	int32 GetBadPacketCounter() const		{ return BadPacketCounter; }
+	void IncErrorPacketCounter()			{ ++ErrorPacketCounter; }
+	int32 GetErrorPacketCounter() const		{ return ErrorPacketCounter; }
+	void IncDroppedPacketCounter(int32 Count=1)			{ DroppedPacketCounter += Count; }
+	int32 GetDroppedPacketCounter() const	{ return DroppedPacketCounter; }
+
+protected:
+	/**
+	 * Performs periodic checks on trigger quota's and packet limits, for non-NetConnection packets
+	 *
+	 * @return	Whether or not non-NetConnection packet limits have been reached
+	 */
+	bool CheckNonConnQuotasAndLimits();
+
+	/**
+	 * Performs periodic checks on NetConnection packet limits
+	 *
+	 * @return	Whether or not NetColnnection packet limits have been reached
+	 */
+	bool CheckNetConnLimits()
+	{
+		return NetConnPacketTimeLimitMSPerFrame > 0 &&
+				(int32)((FPlatformTime::Seconds() - StartFrameRecvTimestamp) * 1000.0) > NetConnPacketTimeLimitMSPerFrame;
+	}
+
+
+
+protected:
+	/** Whether or not DDoS detection is presently enabled */
+	bool bDDoSDetection;
+
+	/** Whether or not analytics for DDoS detection is enabled */
+	bool bDDoSAnalytics;
+
+	/** Whether or not the current frame has reached non-NetConnection packet limits, and should block non-NetConnection packets */
+	bool bHitFrameNonConnLimit;
+
+	/** Whether or not the current frame has reached NetConnection packet limits, and should block ALL further packets */
+	bool bHitFrameNetConnLimit;
+
+
+	/** The different DDoS detection states, of escalating severity, depending on the strength of the DDoS */
+	TArray<FDDoSStateConfig> DetectionSeverity;
+
+	/** The currently active DDoS severity state settings */
+	int8 ActiveState;
+
+	/** The worst DDoS severity state that has been active - used for limiting analytics events */
+	int8 WorstActiveState;
+
+	/** The last time the previous severity states escalation conditions were met (to prevent bouncing up/down between states) */
+	double LastMetEscalationConditions;
+
+	/** Limit checking previous states escalation conditions to once per frame */
+	bool bMetEscalationConditionsThisFrame;
+
+
+	/** Whether or not restriction of log messages from non-NetConnection packets is enabled */
+	bool bDDoSLogRestrictions;
+
+	/** The maximum number of non-NetConnection triggered log messages per frame, before further logs are dropped this frame */
+	int32 DDoSLogSpamLimit;
+
+	/** Counter for log restriction hits, in the current frame */
+	int32 LogHitCounter;
+
+
+	/** The amount of time since the previous frame, for detecting frame hitches, to prevent DDoS detection false positives */
+	int32 HitchTimeQuotaMS;
+
+	/** The number of frames spent hitching, before disabling false positive detection, and treating packet buildup as potential DDoS */
+	int8 HitchFrameTolerance;
+
+	/** The number of consecutive frames spent hitching */
+	int32 HitchFrameCount;
+
+
+	/** Timestamp for the last time per-second quota counting began */
+	double LastPerSecQuotaBegin;
+
+	/** Stores enough per second quota history, to allow all DetectionSeverity states to recalculate if their CooloffTime is reached */
+	TArray<FDDoSPacketCounters> CounterPerSecHistory;
+
+	/** The last written index of CounterPerSecHistory */
+	int32 LastCounterPerSecHistoryIdx;
+
+
+	/** The timestamp for the start of the current frames receive */
+	double StartFrameRecvTimestamp;
+
+	/** Timestamp for the end of the last frames receive loop */
+	double EndFrameRecvTimestamp;
+
+	/** Counts the packets from the start of the current frame */
+	int32 StartFramePacketCount;
+
+	/** The expected time between frames (1.0 / MaxTickRate) - used for adjusting limits/quota's based on DeltaTime */
+	double ExpectedFrameTime;
+
+	/** The current frames adjustment/deviation, from ExpectedFrameTime */
+	float FrameAdjustment;
+
+
+public:
+	/** Analytics delegate for notifying of severity state escalations */
+	FDDoSSeverityEscalation NotifySeverityEscalation;
+};
diff --git a/Engine/Source/Runtime/Net/Core/Public/Net/Core/Misc/PacketAudit.h b/Engine/Source/Runtime/Net/Core/Public/Net/Core/Misc/PacketAudit.h
new file mode 100644
--- /dev/null
+++ b/Engine/Source/Runtime/Net/Core/Public/Net/Core/Misc/PacketAudit.h
@@ -0,0 +1,262 @@
+// Copyright 1998-2019 Epic Games, Inc. All Rights Reserved.
+
+#pragma once
+
+
+/**
+ * Packet Auditing
+ *
+ * Description:
+ *	This implements a means of verifying the correct serialization of packets between a client and server,
+ *	and is primarily for use during development of the netcode (particularly with the PacketHandler).
+ *
+ *	If used liberally, it should quickly zoom in on the precise location of packet serialization problems,
+ *	greatly reducing time spent debugging, and - critically - avoiding having to manually reverse engineer packet dumps.
+ *
+ *
+ * Basic usage:
+ *	In code where you are writing data for sending, you can mark a stage of packet writing for auditing, like so:
+ *		FPacketAudit::AddStage(TEXT("UniqueName"), OutPacket);
+ *
+ *	In code where you are reading incoming packet data, you can implement auditing for that stage as so:
+ *		FPacketAudit::CheckStage(TEXT("UniqueName"), InPacket);
+ *
+ *	Replace 'UniqueName', with a descriptive and unique name, for marking the stage of serialization being checked.
+ *	OutPacket must refer to an FBitWriter, and InPacket must refer to an FBitReader.
+ *
+ *	Run the client/server with '-PacketAudit' added to the commandline.
+ *
+ *
+ * Notes/Limitations:
+ *	- It's safe to leave AddStage/CheckStage in the netcode, without wrapping them in #if blocks, as they're compiled-out in shipping
+ *	- You need to run the server/client as administrator
+ *	- Only a single client and server can perform packet auditing at any given time.
+ *	- Does not function well with ReliabilityHandlerComponent or packetloss.
+ *	- Auditing assumes simple bit reader/writer serialization, and may choke when used with more complex packet parsing
+ *	- Every call to AddStage, must have a corresponding call to CheckStage - or packet data will leak and eventually cause an assert
+ *
+ *
+ * Technicals:
+ *	The auditing code uses Inter-Process-Communication, by means of shared memory and a mutex for guarding access,
+ *	to serialize two TMap's containing packet auditing data, into binary format, and passes them between the client/server processes.
+ */
+
+
+// Includes
+#include "HAL/PlatformProcess.h"
+#include "Serialization/BitReader.h"
+#include "Serialization/BitWriter.h"
+
+
+// Forward declarations
+class FPacketAudit;
+
+
+// Globals
+
+#if !UE_BUILD_SHIPPING
+/** The handler for managing packet auditing */
+extern PACKETHANDLER_API FPacketAudit* GPacketAuditor;
+#endif
+
+
+/**
+ * FPacketAudit
+ *
+ * Uses inter-process-communication, to audit every stage of packet processing, for verifying that packets are processing properly,
+ * and that the client/server are in sync at every stage of reading/writing.
+ *
+ * Since this relies on IPC, it is only suitable for testing on a single machine.
+ *
+ * IMPORTANT: Only supports interaction between a single server, and a single client.
+ */
+class PACKETHANDLER_API FPacketAudit
+{
+#if !UE_BUILD_SHIPPING
+protected:
+	/**
+	 * Provides scoped exclusive access to GSendPacketAudits and GReceivePacketAudits within the .cpp file
+	 */
+	struct FScopedAuditAccess
+	{
+		FScopedAuditAccess(FPacketAudit& InAudtor);
+
+		~FScopedAuditAccess();
+
+	private:
+		FPacketAudit& Auditor;
+	};
+#endif
+
+private:
+	/**
+	 * Default constructor
+	 */
+	FPacketAudit()
+	{
+	}
+
+	/**
+	 * Default constructor
+	 */
+	FPacketAudit(FPlatformProcess::FSemaphore* InGameMutex);
+
+	/**
+	 * Default destructor
+	 */
+	~FPacketAudit();
+
+public:
+	/**
+	 * Initialize GPacketAudtor
+	 */
+	static void Init();
+
+	/**
+	 * Destroy GPacketAuditor
+	 */
+	static void Destruct()
+	{
+#if !UE_BUILD_SHIPPING
+		if (GPacketAuditor != nullptr)
+		{
+			delete GPacketAuditor;
+			GPacketAuditor = nullptr;
+		}
+#endif
+	}
+
+	/**
+	 * On the send side, marks a named stage during packet writing, which should be audited on the receive side.
+	 *
+	 * @param StageName		The unique name to provide this stage of packet writing
+	 * @param OutPacket		The packet the data is being written to - used for gathering audit information
+	 * @param bByteAligned	Whether or not to treat the packet as if it is byte aligned (used in rare circumstances)
+	 */
+	static FORCEINLINE void AddStage(FString StageName, FBitWriter& OutPacket, bool bByteAligned=false)
+	{
+#if !UE_BUILD_SHIPPING
+		if (GPacketAuditor != nullptr)
+		{
+			GPacketAuditor->AddStage_Internal(StageName, OutPacket, bByteAligned);
+		}
+#endif
+	}
+
+	/**
+	 * On the receive side, checks to see that a named stage was marked for auditing on the send side, and audits the packet data.
+	 *
+	 * @param StageName		The unique name to this stage of packet processing
+	 * @param InPacket		The packet the data is being read from - used for gathering audit information
+	 * @param bByteAligned	Whether or not to treat the packet as if it is byte aligned (used in rare circumstances)
+	 */
+	static FORCEINLINE void CheckStage(FString StageName, FBitReader& InPacket, bool bByteAligned=false)
+	{
+#if !UE_BUILD_SHIPPING
+		if (GPacketAuditor != nullptr)
+		{
+			GPacketAuditor->CheckStage_Internal(StageName, InPacket, bByteAligned);
+		}
+#endif
+	}
+
+	/**
+	 * Low level netcode hook, notifying of outgoing packets
+	 *
+	 * @param OutPacket		The current outgoing packet
+	 */
+	static FORCEINLINE void NotifyLowLevelSend(FBitWriter& OutPacket)
+	{
+#if !UE_BUILD_SHIPPING
+		if (GPacketAuditor != nullptr)
+		{
+			GPacketAuditor->NotifyLowLevelSend_Internal(OutPacket);
+		}
+#endif
+	}
+
+	/**
+	 * Low level netcode hook, notifying of incoming packets
+	 *
+	 * @param InPacket		The current incoming packet
+	 */
+	static FORCEINLINE void NotifyLowLevelReceive(FBitReader& InPacket)
+	{
+#if !UE_BUILD_SHIPPING
+		if (GPacketAuditor != nullptr)
+		{
+			GPacketAuditor->NotifyLowLevelReceive_Internal(InPacket);
+		}
+#endif
+	}
+
+	// @todo #JohnB: Deprecate, once the pipeline uses FBitReader/Writers all the way
+#if 1
+	static FORCEINLINE void NotifyLowLevelSend(uint8* Data, uint32 DataLen, uint32 DataLenBits)
+	{
+#if !UE_BUILD_SHIPPING
+		if (GPacketAuditor != nullptr)
+		{
+			FBitWriter OutPacket(DataLenBits);
+
+			OutPacket.SerializeBits(Data, DataLenBits);
+
+			GPacketAuditor->NotifyLowLevelSend_Internal(OutPacket);
+		}
+#endif
+	}
+
+	static FORCEINLINE void NotifyLowLevelReceive(uint8* Data, uint32 DataLen)
+	{
+#if !UE_BUILD_SHIPPING
+		if (GPacketAuditor != nullptr)
+		{
+			FBitReader InPacket(Data, DataLen * 8);
+
+			NotifyLowLevelReceive(InPacket);
+		}
+#endif
+	}
+#endif
+
+protected:
+	void AddStage_Internal(FString StageName, FBitWriter& OutPacket, bool bByteAligned);
+
+	void CheckStage_Internal(FString StageName, FBitReader& InPacket, bool bByteAligned);
+
+	void NotifyLowLevelSend_Internal(FBitWriter& OutPacket);
+
+	void NotifyLowLevelReceive_Internal(FBitReader& InPacket);
+
+
+	/**
+	 * Gets the CRC of a packet, zeroing any unwritten bits in the last byte beforehand.
+	 *
+	 * @param Data			The packet data
+	 * @param DataLenBits	The length of the data, in bits
+	 * @return				Returns the CRC of the packet
+	 */
+	static uint32 PacketCRC(uint8* Data, uint32 DataLenBits);
+
+#if !UE_BUILD_SHIPPING
+	/**
+	 * Dumps the collected packet audit data to the log
+	 */
+	void DumpAuditData(FScopedAuditAccess& AuditLock, uint32 InPacketCRC=0);
+#endif
+
+
+protected:
+	/** Mutex used for enforcing a single instance of client/server packet auditing */
+	FPlatformProcess::FSemaphore*			GameMutex;
+
+	/** Mutex used for exclusive access to shared memory */
+	FPlatformProcess::FSemaphore*			SharedMutex;
+
+	/** Shared memory region for packet sends */
+	FPlatformMemory::FSharedMemoryRegion*	SendSharedMemory;
+
+	/** Shared memory region for packet receives */
+	FPlatformMemory::FSharedMemoryRegion*	ReceiveSharedMemory;
+};
+
diff --git a/Engine/Source/Runtime/Net/Core/Public/Net/Core/Misc/PacketTraits.h b/Engine/Source/Runtime/Net/Core/Public/Net/Core/Misc/PacketTraits.h
new file mode 100644
--- /dev/null
+++ b/Engine/Source/Runtime/Net/Core/Public/Net/Core/Misc/PacketTraits.h
@@ -0,0 +1,47 @@
+// Copyright 1998-2019 Epic Games, Inc. All Rights Reserved.
+
+#pragma once
+
+// includes
+#include "CoreMinimal.h"
+
+
+// @todo #JohnB: Create a new CoreNet module, remove conflicts with existing 'CoreNet' files, and begin moving stuff like this into it
+
+
+/**
+ * Contains metadata and flags, which provide information on the traits of a packet - what it contains and how to process it.
+ */
+struct PACKETHANDLER_API FOutPacketTraits
+{
+	/** Flags - may trigger modification of packet and traits */
+
+	/** Whether or not the packet can/should be compressed */
+	bool bAllowCompression;
+
+
+	/** Traits */
+
+	/** The number of ack bits in the packet - reflecting UNetConnection.NumAckBits */
+	uint32 NumAckBits;
+
+	/** The number of bunch bits in the packet - reflecting UNetConnection.NumBunchBits */
+	uint32 NumBunchBits;
+
+	/** Whether or not this is a keepalive packet */
+	bool bIsKeepAlive;
+
+	/** Whether or not the packet has been compressed */
+	bool bIsCompressed;
+
+
+	/** Default constructor */
+	FOutPacketTraits()
+		: bAllowCompression(true)
+		, NumAckBits(0)
+		, NumBunchBits(0)
+		, bIsKeepAlive(false)
+		, bIsCompressed(false)
+	{
+	}
+};
diff --git a/Engine/Source/Runtime/PacketHandlers/PacketHandler/Classes/NetAnalyticsAggregatorConfig.h b/Engine/Source/Runtime/PacketHandlers/PacketHandler/Classes/NetAnalyticsAggregatorConfig.h
deleted file mode 100644
diff --git a/Engine/Source/Runtime/PacketHandlers/PacketHandler/Private/DDoSDetection.cpp b/Engine/Source/Runtime/PacketHandlers/PacketHandler/Private/DDoSDetection.cpp
deleted file mode 100644
diff --git a/Engine/Source/Runtime/PacketHandlers/PacketHandler/Private/NetAnalytics.cpp b/Engine/Source/Runtime/PacketHandlers/PacketHandler/Private/NetAnalytics.cpp
deleted file mode 100644
diff --git a/Engine/Source/Runtime/PacketHandlers/PacketHandler/Private/NetAnalyticsAggregatorConfig.cpp b/Engine/Source/Runtime/PacketHandlers/PacketHandler/Private/NetAnalyticsAggregatorConfig.cpp
deleted file mode 100644
diff --git a/Engine/Source/Runtime/PacketHandlers/PacketHandler/Private/PacketAudit.cpp b/Engine/Source/Runtime/PacketHandlers/PacketHandler/Private/PacketAudit.cpp
deleted file mode 100644
diff --git a/Engine/Source/Runtime/PacketHandlers/PacketHandler/Public/DDoSDetection.h b/Engine/Source/Runtime/PacketHandlers/PacketHandler/Public/DDoSDetection.h
deleted file mode 100644
diff --git a/Engine/Source/Runtime/PacketHandlers/PacketHandler/Public/NetAnalytics.h b/Engine/Source/Runtime/PacketHandlers/PacketHandler/Public/NetAnalytics.h
deleted file mode 100644
diff --git a/Engine/Source/Runtime/PacketHandlers/PacketHandler/Public/PacketAudit.h b/Engine/Source/Runtime/PacketHandlers/PacketHandler/Public/PacketAudit.h
deleted file mode 100644
diff --git a/Engine/Source/Runtime/PacketHandlers/PacketHandler/Public/PacketTraits.h b/Engine/Source/Runtime/PacketHandlers/PacketHandler/Public/PacketTraits.h
deleted file mode 100644
