Added NetCore module.

diff --git a/Engine/Plugins/Online/OnlineSubsystemOculus/Source/OnlineSubsystemOculus.Build.cs b/Engine/Plugins/Online/OnlineSubsystemOculus/Source/OnlineSubsystemOculus.Build.cs
--- a/Engine/Plugins/Online/OnlineSubsystemOculus/Source/OnlineSubsystemOculus.Build.cs
+++ b/Engine/Plugins/Online/OnlineSubsystemOculus/Source/OnlineSubsystemOculus.Build.cs
@@ -22,6 +22,7 @@ public class OnlineSubsystemOculus : ModuleRules
 		PrivateDependencyModuleNames.AddRange(
 			new string[] {
 				"Core",
+				"NetCore",
 				"CoreUObject",
 				"Engine",
 				"Sockets",
diff --git a/Engine/Plugins/Online/OnlineSubsystemSteam/Source/OnlineSubsystemSteam.Build.cs b/Engine/Plugins/Online/OnlineSubsystemSteam/Source/OnlineSubsystemSteam.Build.cs
--- a/Engine/Plugins/Online/OnlineSubsystemSteam/Source/OnlineSubsystemSteam.Build.cs
+++ b/Engine/Plugins/Online/OnlineSubsystemSteam/Source/OnlineSubsystemSteam.Build.cs
@@ -19,7 +19,8 @@ public class OnlineSubsystemSteam : ModuleRules
 		PrivateDependencyModuleNames.AddRange(
 			new string[] {
 				"Core", 
-				"CoreUObject", 
+				"CoreUObject",
+				"NetCore",
 				"Engine", 
 				"Sockets", 
 				"Voice",
diff --git a/Engine/Plugins/Online/OnlineSubsystemUtils/Source/OnlineSubsystemUtils/OnlineSubsystemUtils.Build.cs b/Engine/Plugins/Online/OnlineSubsystemUtils/Source/OnlineSubsystemUtils/OnlineSubsystemUtils.Build.cs
--- a/Engine/Plugins/Online/OnlineSubsystemUtils/Source/OnlineSubsystemUtils/OnlineSubsystemUtils.Build.cs
+++ b/Engine/Plugins/Online/OnlineSubsystemUtils/Source/OnlineSubsystemUtils/OnlineSubsystemUtils.Build.cs
@@ -42,6 +42,7 @@ public class OnlineSubsystemUtils : ModuleRules
 			{
 				"Core",
 				"CoreUObject",
+				"NetCore",
 				"Engine"
 			}
 		);
diff --git a/Engine/Plugins/Online/OnlineSubsystemUtils/Source/OnlineSubsystemUtils/Private/IpConnection.cpp b/Engine/Plugins/Online/OnlineSubsystemUtils/Source/OnlineSubsystemUtils/Private/IpConnection.cpp
--- a/Engine/Plugins/Online/OnlineSubsystemUtils/Source/OnlineSubsystemUtils/Private/IpConnection.cpp
+++ b/Engine/Plugins/Online/OnlineSubsystemUtils/Source/OnlineSubsystemUtils/Private/IpConnection.cpp
@@ -16,7 +16,7 @@ Notes:
 #include "Net/NetworkProfiler.h"
 #include "Net/DataChannel.h"
 
-#include "PacketAudit.h"
+#include "Net/Core/Misc/PacketAudit.h"
 
 /*-----------------------------------------------------------------------------
 	Declarations.
diff --git a/Engine/Plugins/Online/OnlineSubsystemUtils/Source/OnlineSubsystemUtils/Private/IpNetDriver.cpp b/Engine/Plugins/Online/OnlineSubsystemUtils/Source/OnlineSubsystemUtils/Private/IpNetDriver.cpp
--- a/Engine/Plugins/Online/OnlineSubsystemUtils/Source/OnlineSubsystemUtils/Private/IpNetDriver.cpp
+++ b/Engine/Plugins/Online/OnlineSubsystemUtils/Source/OnlineSubsystemUtils/Private/IpNetDriver.cpp
@@ -20,7 +20,7 @@ Notes:
 #include "IpConnection.h"
 #include "HAL/LowLevelMemTracker.h"
 
-#include "PacketAudit.h"
+#include "Net/Core/Misc/PacketAudit.h"
 
 #include "IPAddress.h"
 #include "Sockets.h"
diff --git a/Engine/Plugins/Runtime/PacketHandlers/CompressionComponents/Oodle/Source/OodleHandlerComponent/OodleHandlerComponent.Build.cs b/Engine/Plugins/Runtime/PacketHandlers/CompressionComponents/Oodle/Source/OodleHandlerComponent/OodleHandlerComponent.Build.cs
--- a/Engine/Plugins/Runtime/PacketHandlers/CompressionComponents/Oodle/Source/OodleHandlerComponent/OodleHandlerComponent.Build.cs
+++ b/Engine/Plugins/Runtime/PacketHandlers/CompressionComponents/Oodle/Source/OodleHandlerComponent/OodleHandlerComponent.Build.cs
@@ -17,6 +17,7 @@ public class OodleHandlerComponent : ModuleRules
 				"PacketHandler",
 				"Core",
 				"CoreUObject",
+				"NetCore",
 				"Engine",
                 "Analytics"
 			});
diff --git a/Engine/Plugins/Runtime/PacketHandlers/CompressionComponents/Oodle/Source/OodleHandlerComponent/Public/OodleAnalytics.h b/Engine/Plugins/Runtime/PacketHandlers/CompressionComponents/Oodle/Source/OodleHandlerComponent/Public/OodleAnalytics.h
--- a/Engine/Plugins/Runtime/PacketHandlers/CompressionComponents/Oodle/Source/OodleHandlerComponent/Public/OodleAnalytics.h
+++ b/Engine/Plugins/Runtime/PacketHandlers/CompressionComponents/Oodle/Source/OodleHandlerComponent/Public/OodleAnalytics.h
@@ -4,7 +4,7 @@
 
 // Include
 #include "CoreMinimal.h"
-#include "NetAnalytics.h"
+#include "Net/Core/Analytics/NetAnalytics.h"
 
 
 /**
diff --git a/Engine/Source/Runtime/Engine/Classes/Engine/NetConnection.h b/Engine/Source/Runtime/Engine/Classes/Engine/NetConnection.h
--- a/Engine/Source/Runtime/Engine/Classes/Engine/NetConnection.h
+++ b/Engine/Source/Runtime/Engine/Classes/Engine/NetConnection.h
@@ -22,7 +22,7 @@
 #include "Containers/CircularBuffer.h"
 #include "ReplicationDriver.h"
 #include "Analytics/EngineNetAnalytics.h"
-#include "PacketTraits.h"
+#include "Net/Core/Misc/PacketTraits.h"
 #include "Net/Util/ResizableCircularQueue.h"
 #include "Net/NetAnalyticsTypes.h"
 
diff --git a/Engine/Source/Runtime/Engine/Classes/Engine/NetDriver.h b/Engine/Source/Runtime/Engine/Classes/Engine/NetDriver.h
--- a/Engine/Source/Runtime/Engine/Classes/Engine/NetDriver.h
+++ b/Engine/Source/Runtime/Engine/Classes/Engine/NetDriver.h
@@ -14,7 +14,7 @@
 #include "GameFramework/WorldSettings.h"
 #include "PacketHandler.h"
 #include "Channel.h"
-#include "DDoSDetection.h"
+#include "Net/Core/Misc/DDoSDetection.h"
 #include "IPAddress.h"
 #include "Net/NetAnalyticsTypes.h"
 
diff --git a/Engine/Source/Runtime/Engine/Engine.Build.cs b/Engine/Source/Runtime/Engine/Engine.Build.cs
--- a/Engine/Source/Runtime/Engine/Engine.Build.cs
+++ b/Engine/Source/Runtime/Engine/Engine.Build.cs
@@ -66,6 +66,7 @@ public class Engine : ModuleRules
 			new string[] {
 				"Core",
 				"CoreUObject",
+				"NetCore",
 				"ApplicationCore",
 				"Json",
 				"SlateCore",
diff --git a/Engine/Source/Runtime/Engine/Private/NetworkDriver.cpp b/Engine/Source/Runtime/Engine/Private/NetworkDriver.cpp
--- a/Engine/Source/Runtime/Engine/Private/NetworkDriver.cpp
+++ b/Engine/Source/Runtime/Engine/Private/NetworkDriver.cpp
@@ -37,7 +37,7 @@
 #include "GameFramework/WorldSettings.h"
 #include "PacketHandler.h"
 #include "PacketHandlers/StatelessConnectHandlerComponent.h"
-#include "NetAnalytics.h"
+#include "Net/Core/Analytics/NetAnalytics.h"
 #include "Engine/NetDriver.h"
 #include "Engine/LocalPlayer.h"
 #include "Net/DataBunch.h"
diff --git a/Engine/Source/Runtime/Engine/Private/PacketHandlers/StatelessConnectHandlerComponent.cpp b/Engine/Source/Runtime/Engine/Private/PacketHandlers/StatelessConnectHandlerComponent.cpp
--- a/Engine/Source/Runtime/Engine/Private/PacketHandlers/StatelessConnectHandlerComponent.cpp
+++ b/Engine/Source/Runtime/Engine/Private/PacketHandlers/StatelessConnectHandlerComponent.cpp
@@ -6,7 +6,7 @@
 #include "EngineStats.h"
 #include "Misc/SecureHash.h"
 #include "Engine/NetConnection.h"
-#include "PacketAudit.h"
+#include "Net/Core/Misc/PacketAudit.h"
 
 
 
diff --git a/Engine/Source/Runtime/Engine/Public/Analytics/EngineNetAnalytics.h b/Engine/Source/Runtime/Engine/Public/Analytics/EngineNetAnalytics.h
--- a/Engine/Source/Runtime/Engine/Public/Analytics/EngineNetAnalytics.h
+++ b/Engine/Source/Runtime/Engine/Public/Analytics/EngineNetAnalytics.h
@@ -3,7 +3,7 @@
 #pragma once
 
 // Includes
-#include "NetAnalytics.h"
+#include "Net/Core/Analytics/NetAnalytics.h"
 
 
 /**
diff --git a/Engine/Source/Runtime/Net/Core/Private/Net/Core/Analytics/NetAnalytics.cpp b/Engine/Source/Runtime/Net/Core/Private/Net/Core/Analytics/NetAnalytics.cpp
--- a/Engine/Source/Runtime/Net/Core/Private/Net/Core/Analytics/NetAnalytics.cpp
+++ b/Engine/Source/Runtime/Net/Core/Private/Net/Core/Analytics/NetAnalytics.cpp
@@ -1,9 +1,9 @@
 // Copyright 1998-2019 Epic Games, Inc. All Rights Reserved.
 
 // Includes
-#include "NetAnalytics.h"
-#include "PacketHandler.h"
-#include "NetAnalyticsAggregatorConfig.h"
+#include "Net/Core/Analytics/NetAnalytics.h"
+#include "Net/Core/Analytics/NetAnalyticsAggregatorConfig.h"
+#include "Net/Core/Misc/NetCoreLog.h"
 
 
 // Globals
@@ -81,7 +81,7 @@ void FNetAnalyticsAggregator::InitConfig()
 
 	// If the config is hotfixed, make sure no data holders are currently active, as they can't be selectively hotfixed if loaded
 	// (this does seem to happen, frequently - so limits the hotfixability of Net Analytics)
-	UE_CLOG(AnalyticsDataMap.Num() > 0, PacketHandlerLog, Warning,
+	UE_CLOG(AnalyticsDataMap.Num() > 0, LogNetCore, Warning,
 			TEXT("Net Analytics hotfixed while already active. Analytics hotfix changes may not be applied correctly."));
 
 	if (CurConfig != nullptr)
@@ -90,7 +90,7 @@ void FNetAnalyticsAggregator::InitConfig()
 		{
 			AnalyticsDataConfigMap.Add(CurEntry.DataName, CurEntry.bEnabled);
 
-			UE_LOG(PacketHandlerLog, Log, TEXT("Adding NetAnalyticsData: %s, bEnabled: %i"), *CurEntry.DataName.ToString(),
+			UE_LOG(LogNetCore, Log, TEXT("Adding NetAnalyticsData: %s, bEnabled: %i"), *CurEntry.DataName.ToString(),
 					CurEntry.bEnabled);
 		}
 	}
@@ -123,7 +123,7 @@ TNetAnalyticsDataPtr<> FNetAnalyticsAggregator::RegisterAnalyticsData_Internal(T
 	}
 	else if (bEnabled == nullptr)
 	{
-		UE_LOG(PacketHandlerLog, Error, TEXT("NetAnalyticsData type '%s' must be added to NetAnalyticsAggregatorConfig, for NetDriverName: %s."),
+		UE_LOG(LogNetCore, Error, TEXT("NetAnalyticsData type '%s' must be added to NetAnalyticsAggregatorConfig, for NetDriverName: %s."),
 				*InDataName.ToString(), *NetDriverName.ToString());
 	}
 
diff --git a/Engine/Source/Runtime/Net/Core/Private/Net/Core/Analytics/NetAnalyticsAggregatorConfig.cpp b/Engine/Source/Runtime/Net/Core/Private/Net/Core/Analytics/NetAnalyticsAggregatorConfig.cpp
--- a/Engine/Source/Runtime/Net/Core/Private/Net/Core/Analytics/NetAnalyticsAggregatorConfig.cpp
+++ b/Engine/Source/Runtime/Net/Core/Private/Net/Core/Analytics/NetAnalyticsAggregatorConfig.cpp
@@ -1,6 +1,6 @@
 // Copyright 1998-2019 Epic Games, Inc. All Rights Reserved.
 
-#include "NetAnalyticsAggregatorConfig.h"
+#include "Net/Core/Analytics/NetAnalyticsAggregatorConfig.h"
 
 
 /**
diff --git a/Engine/Source/Runtime/Net/Core/Private/Net/Core/Misc/DDoSDetection.cpp b/Engine/Source/Runtime/Net/Core/Private/Net/Core/Misc/DDoSDetection.cpp
--- a/Engine/Source/Runtime/Net/Core/Private/Net/Core/Misc/DDoSDetection.cpp
+++ b/Engine/Source/Runtime/Net/Core/Private/Net/Core/Misc/DDoSDetection.cpp
@@ -2,8 +2,8 @@
 
 // Includes
 
-#include "DDoSDetection.h"
-#include "PacketHandler.h"
+#include "Net/Core/Misc/DDoSDetection.h"
+#include "Net/Core/Misc/NetCoreLog.h"
 #include "Misc/ConfigCacheIni.h"
 
 
@@ -113,7 +113,7 @@ void FDDoSDetection::InitConfig()
 
 	DetectionSeverity.Empty();
 
-	UE_LOG(PacketHandlerLog, Log, TEXT("DDoS detection status: detection enabled: %d analytics enabled: %d"), bDDoSDetection, bDDoSAnalytics);
+	UE_LOG(LogNetCore, Log, TEXT("DDoS detection status: detection enabled: %d analytics enabled: %d"), bDDoSDetection, bDDoSAnalytics);
 
 	if (bDDoSDetection)
 	{
@@ -150,7 +150,7 @@ void FDDoSDetection::InitConfig()
 			}
 			else
 			{
-				UE_LOG(PacketHandlerLog, Warning, TEXT("DDoS detection could not find ini section: %s"), *CurSection);
+				UE_LOG(LogNetCore, Warning, TEXT("DDoS detection could not find ini section: %s"), *CurSection);
 			}
 		}
 
@@ -162,7 +162,7 @@ void FDDoSDetection::InitConfig()
 		}
 		else
 		{
-			UE_LOG(PacketHandlerLog, Warning, TEXT("DDoS detection enabled, but no DetectionSeverity states specified! Disabling."));
+			UE_LOG(LogNetCore, Warning, TEXT("DDoS detection enabled, but no DetectionSeverity states specified! Disabling."));
 
 			bDDoSDetection = false;
 		}
@@ -238,7 +238,7 @@ void FDDoSDetection::UpdateSeverity(bool bEscalate)
 		}
 
 
-		UE_LOG(PacketHandlerLog, Warning, TEXT("Updated DDoS detection severity from '%s' to '%s'"),
+		UE_LOG(LogNetCore, Warning, TEXT("Updated DDoS detection severity from '%s' to '%s'"),
 				*OldState.SeverityCategory, *CurState.SeverityCategory);
 
 		if (bEscalate && ActiveState > WorstActiveState)
@@ -268,7 +268,7 @@ void FDDoSDetection::PreFrameReceive(float DeltaTime)
 			{
 				HitchFrameCount++;
 
-				UE_LOG(PacketHandlerLog, Verbose, TEXT("Detected '%i' successive hitches outside NetDriver Tick. Last Hitch: %fms (Max: %ims)"),
+				UE_LOG(LogNetCore, Verbose, TEXT("Detected '%i' successive hitches outside NetDriver Tick. Last Hitch: %fms (Max: %ims)"),
 						HitchFrameCount, HitchTimeMS, HitchTimeQuotaMS);
 			}
 			else
@@ -289,7 +289,7 @@ void FDDoSDetection::PreFrameReceive(float DeltaTime)
 
 		if (((StartFrameRecvTimestamp - LastPerSecQuotaBegin) - 1.0) > 0.0)
 		{
-			UE_CLOG(DroppedPacketCounter > 0, PacketHandlerLog, Warning,
+			UE_CLOG(DroppedPacketCounter > 0, LogNetCore, Warning,
 				TEXT("DDoS Detection dropped '%i' packets during last second (bHitFrameNonConnLimit: %i, bHitFrameNetConnLimit: %i, ")
 				TEXT("DetectionSeverity: %s)."),
 				DroppedPacketCounter, (int32)bHitFrameNonConnLimit, (int32)bHitFrameNetConnLimit,
@@ -318,7 +318,7 @@ void FDDoSDetection::PreFrameReceive(float DeltaTime)
 
 		if (LogHitCounter >= DDoSLogSpamLimit)
 		{
-			UE_LOG(PacketHandlerLog, Warning, TEXT("Previous frame hit DDoS LogHitCounter limit - hit count: %i (Max: %i)"), LogHitCounter,
+			UE_LOG(LogNetCore, Warning, TEXT("Previous frame hit DDoS LogHitCounter limit - hit count: %i (Max: %i)"), LogHitCounter,
 					DDoSLogSpamLimit);
 		}
 
diff --git a/Engine/Source/Runtime/Net/Core/Private/Net/Core/Misc/PacketAudit.cpp b/Engine/Source/Runtime/Net/Core/Private/Net/Core/Misc/PacketAudit.cpp
--- a/Engine/Source/Runtime/Net/Core/Private/Net/Core/Misc/PacketAudit.cpp
+++ b/Engine/Source/Runtime/Net/Core/Private/Net/Core/Misc/PacketAudit.cpp
@@ -3,8 +3,8 @@
 /**
  * Includes
  */
-#include "PacketAudit.h"
-#include "PacketHandler.h"
+#include "Net/Core/Misc/PacketAudit.h"
+#include "Net/Core/Misc/NetCoreLog.h"
 
 #include "Misc/CommandLine.h"
 #include "Serialization/BitReader.h"
@@ -296,7 +296,7 @@ void FPacketAudit::Init()
 		}
 		else
 		{
-			UE_LOG(PacketHandlerLog, Log, TEXT("%s"), *(FString(TEXT("Packet auditor already active for a game ")) +
+			UE_LOG(LogNetCore, Log, TEXT("%s"), *(FString(TEXT("Packet auditor already active for a game ")) +
 					(GIsServer ? TEXT("server") : TEXT("client")) + TEXT(", can't start multiple instances.")));
 
 			FPlatformProcess::DeleteInterprocessSynchObject(CurGameMutex);
@@ -481,7 +481,7 @@ void FPacketAudit::DumpAuditData(FScopedAuditAccess& AuditLock, uint32 InPacketC
 
 	FinalLogStr += TEXT(")");
 
-	UE_LOG(PacketHandlerLog, Log, TEXT("%s"), *FinalLogStr);
+	UE_LOG(LogNetCore, Log, TEXT("%s"), *FinalLogStr);
 }
 
 
diff --git a/Engine/Source/Runtime/Net/Core/Public/Net/Core/Analytics/NetAnalytics.h b/Engine/Source/Runtime/Net/Core/Public/Net/Core/Analytics/NetAnalytics.h
--- a/Engine/Source/Runtime/Net/Core/Public/Net/Core/Analytics/NetAnalytics.h
+++ b/Engine/Source/Runtime/Net/Core/Public/Net/Core/Analytics/NetAnalytics.h
@@ -87,7 +87,7 @@ template<class T=FNetAnalyticsData> using TNetAnalyticsDataWeakPtr = TWeakPtr<T,
 // Globals
 
 /** Counter used to detect when a new FNetAnalyticsAggregator (i.e. new NetDriver) has been created - to trigger cleanup of TLS data */
-extern PACKETHANDLER_API TAtomic<uint8> GNetAnalyticsCounter;
+extern NETCORE_API TAtomic<uint8> GNetAnalyticsCounter;
 
 
 /**
@@ -104,7 +104,7 @@ struct FLocalNetAnalyticsStruct : public FVirtualDestructor
 /**
  * Subclassed struct which holds net analytics data which is aggregated or to-be-aggregated, and dispatched upon SendAnalytics
  */
-struct PACKETHANDLER_API FNetAnalyticsData : public FVirtualDestructor, public TSharedFromThis<FNetAnalyticsData, NetAnalyticsThreadSafety>
+struct NETCORE_API FNetAnalyticsData : public FVirtualDestructor, public TSharedFromThis<FNetAnalyticsData, NetAnalyticsThreadSafety>
 {
 	friend class FNetAnalyticsAggregator;
 	friend struct FNetAnalyticsDataDeleter;
@@ -175,7 +175,7 @@ template<class TDataStruct> struct TBasicNetAnalyticsData : public FNetAnalytics
  *
  * The last thread to release this analytics data, triggers SendAnalytics - thread safety is partially provided by shared pointer atomics.
  */
-struct PACKETHANDLER_API FThreadedNetAnalyticsData : public FNetAnalyticsData
+struct NETCORE_API FThreadedNetAnalyticsData : public FNetAnalyticsData
 {
 	/**
 	 * Default constructor
@@ -320,7 +320,7 @@ struct FNetAnalyticsDataDeleter
 /**
  * Central object (usually within NetDriver) which handles registration/retrieval/type-checking of net analytics data holders.
  */
-class PACKETHANDLER_API FNetAnalyticsAggregator
+class NETCORE_API FNetAnalyticsAggregator
 {
 public:
 	/**
diff --git a/Engine/Source/Runtime/Net/Core/Public/Net/Core/Misc/DDoSDetection.h b/Engine/Source/Runtime/Net/Core/Public/Net/Core/Misc/DDoSDetection.h
--- a/Engine/Source/Runtime/Net/Core/Public/Net/Core/Misc/DDoSDetection.h
+++ b/Engine/Source/Runtime/Net/Core/Public/Net/Core/Misc/DDoSDetection.h
@@ -19,7 +19,7 @@ DECLARE_DELEGATE_OneParam(FDDoSSeverityEscalation, FString /*SeverityCategory*/)
 /**
  * Struct containing the per-second packet counters
  */
-struct PACKETHANDLER_API FDDoSPacketCounters
+struct NETCORE_API FDDoSPacketCounters
 {
 	/** Counter for non-NetConnection packets received, since the last per second quota period began */
 	int32 NonConnPacketCounter;
@@ -59,7 +59,7 @@ struct PACKETHANDLER_API FDDoSPacketCounters
 /**
  * Stores the DDoS detection state (either settings from the config file, or the active DDoS detection state)
  */
-struct PACKETHANDLER_API FDDoSState
+struct NETCORE_API FDDoSState
 {
 	/** The number of packets/sec before the next stage of DDoS detection is triggered */
 	int32 EscalateQuotaPacketsPerSec;
@@ -119,7 +119,7 @@ struct PACKETHANDLER_API FDDoSState
 /**
  * DDoS detection state, with functions for applying the state to active DDoS detection
  */
-struct PACKETHANDLER_API FDDoSStateConfig : public FDDoSState
+struct NETCORE_API FDDoSStateConfig : public FDDoSState
 {
 	/** The name of the DDoS severity level this config section represents */
 	FString SeverityCategory;
@@ -161,7 +161,7 @@ struct PACKETHANDLER_API FDDoSStateConfig : public FDDoSState
  * The main DDoS detection tracking class, for counting packets and applying restrictions.
  * Implemented separate to the NetDriver, to allow wider use e.g. potentially at socket level, if useful.
  */
-class PACKETHANDLER_API FDDoSDetection : protected FDDoSPacketCounters, protected FDDoSState
+class NETCORE_API FDDoSDetection : protected FDDoSPacketCounters, protected FDDoSState
 {
 public:
 	/**
diff --git a/Engine/Source/Runtime/Net/Core/Public/Net/Core/Misc/PacketAudit.h b/Engine/Source/Runtime/Net/Core/Public/Net/Core/Misc/PacketAudit.h
--- a/Engine/Source/Runtime/Net/Core/Public/Net/Core/Misc/PacketAudit.h
+++ b/Engine/Source/Runtime/Net/Core/Public/Net/Core/Misc/PacketAudit.h
@@ -56,7 +56,7 @@ class FPacketAudit;
 
 #if !UE_BUILD_SHIPPING
 /** The handler for managing packet auditing */
-extern PACKETHANDLER_API FPacketAudit* GPacketAuditor;
+extern NETCORE_API FPacketAudit* GPacketAuditor;
 #endif
 
 
@@ -70,7 +70,7 @@ extern PACKETHANDLER_API FPacketAudit* GPacketAuditor;
  *
  * IMPORTANT: Only supports interaction between a single server, and a single client.
  */
-class PACKETHANDLER_API FPacketAudit
+class NETCORE_API FPacketAudit
 {
 #if !UE_BUILD_SHIPPING
 protected:
diff --git a/Engine/Source/Runtime/Net/Core/Public/Net/Core/Misc/PacketTraits.h b/Engine/Source/Runtime/Net/Core/Public/Net/Core/Misc/PacketTraits.h
--- a/Engine/Source/Runtime/Net/Core/Public/Net/Core/Misc/PacketTraits.h
+++ b/Engine/Source/Runtime/Net/Core/Public/Net/Core/Misc/PacketTraits.h
@@ -6,13 +6,12 @@
 #include "CoreMinimal.h"
 
 
-// @todo #JohnB: Create a new CoreNet module, remove conflicts with existing 'CoreNet' files, and begin moving stuff like this into it
 
 
 /**
  * Contains metadata and flags, which provide information on the traits of a packet - what it contains and how to process it.
  */
-struct PACKETHANDLER_API FOutPacketTraits
+struct NETCORE_API FOutPacketTraits
 {
 	/** Flags - may trigger modification of packet and traits */
 
diff --git a/Engine/Source/Runtime/PacketHandlers/EncryptionComponents/AsymmetricEncryption/RSAEncryptionHandlerComponent/Private/RSAEncryptionHandlerComponent.cpp b/Engine/Source/Runtime/PacketHandlers/EncryptionComponents/AsymmetricEncryption/RSAEncryptionHandlerComponent/Private/RSAEncryptionHandlerComponent.cpp
--- a/Engine/Source/Runtime/PacketHandlers/EncryptionComponents/AsymmetricEncryption/RSAEncryptionHandlerComponent/Private/RSAEncryptionHandlerComponent.cpp
+++ b/Engine/Source/Runtime/PacketHandlers/EncryptionComponents/AsymmetricEncryption/RSAEncryptionHandlerComponent/Private/RSAEncryptionHandlerComponent.cpp
@@ -2,7 +2,7 @@
 
 #include "RSAEncryptionHandlerComponent.h"
 
-#include "PacketAudit.h"
+#include "Net/Core/Misc/PacketAudit.h"
 
 #include "Modules/ModuleManager.h"
 
diff --git a/Engine/Source/Runtime/PacketHandlers/EncryptionComponents/AsymmetricEncryption/RSAEncryptionHandlerComponent/RSAEncryptionHandlerComponent.Build.cs b/Engine/Source/Runtime/PacketHandlers/EncryptionComponents/AsymmetricEncryption/RSAEncryptionHandlerComponent/RSAEncryptionHandlerComponent.Build.cs
--- a/Engine/Source/Runtime/PacketHandlers/EncryptionComponents/AsymmetricEncryption/RSAEncryptionHandlerComponent/RSAEncryptionHandlerComponent.Build.cs
+++ b/Engine/Source/Runtime/PacketHandlers/EncryptionComponents/AsymmetricEncryption/RSAEncryptionHandlerComponent/RSAEncryptionHandlerComponent.Build.cs
@@ -10,6 +10,7 @@ public class RSAEncryptionHandlerComponent : ModuleRules
         PublicDependencyModuleNames.AddRange(
             new string[] {
 				"Core",
+				"NetCore",
                 "PacketHandler"
             }
         );
diff --git a/Engine/Source/Runtime/PacketHandlers/EncryptionComponents/RSAKeyAESEncryption/Private/RSAKeyAESEncryptionHandlerComponent.cpp b/Engine/Source/Runtime/PacketHandlers/EncryptionComponents/RSAKeyAESEncryption/Private/RSAKeyAESEncryptionHandlerComponent.cpp
--- a/Engine/Source/Runtime/PacketHandlers/EncryptionComponents/RSAKeyAESEncryption/Private/RSAKeyAESEncryptionHandlerComponent.cpp
+++ b/Engine/Source/Runtime/PacketHandlers/EncryptionComponents/RSAKeyAESEncryption/Private/RSAKeyAESEncryptionHandlerComponent.cpp
@@ -2,7 +2,7 @@
 
 #include "RSAKeyAESEncryptionHandlerComponent.h"
 
-#include "PacketAudit.h"
+#include "Net/Core/Misc/PacketAudit.h"
 #include "PacketHandler.h"
 
 #include "Modules/ModuleManager.h"
diff --git a/Engine/Source/Runtime/PacketHandlers/EncryptionComponents/RSAKeyAESEncryption/RSAKeyAESEncryption.Build.cs b/Engine/Source/Runtime/PacketHandlers/EncryptionComponents/RSAKeyAESEncryption/RSAKeyAESEncryption.Build.cs
--- a/Engine/Source/Runtime/PacketHandlers/EncryptionComponents/RSAKeyAESEncryption/RSAKeyAESEncryption.Build.cs
+++ b/Engine/Source/Runtime/PacketHandlers/EncryptionComponents/RSAKeyAESEncryption/RSAKeyAESEncryption.Build.cs
@@ -10,6 +10,7 @@ public class RSAKeyAESEncryption : ModuleRules
         PublicDependencyModuleNames.AddRange(
             new string[] {
 				"Core",
+				"NetCore",
                 "PacketHandler"
             }
         );
diff --git a/Engine/Source/Runtime/PacketHandlers/PacketHandler/PacketHandler.Build.cs b/Engine/Source/Runtime/PacketHandlers/PacketHandler/PacketHandler.Build.cs
--- a/Engine/Source/Runtime/PacketHandlers/PacketHandler/PacketHandler.Build.cs
+++ b/Engine/Source/Runtime/PacketHandlers/PacketHandler/PacketHandler.Build.cs
@@ -9,13 +9,14 @@ public class PacketHandler : ModuleRules
     {
         PublicDependencyModuleNames.AddRange
 		(
-            new string[]
+			new string[]
 			{
 				"Core",
 				"CoreUObject",
-                "ReliabilityHandlerComponent",
+				"NetCore",
+				"ReliabilityHandlerComponent",
 				"Sockets",
-            }
+			}
         );
 
         CircularlyReferencedDependentModules.Add("ReliabilityHandlerComponent");
diff --git a/Engine/Source/Runtime/PacketHandlers/PacketHandler/Private/PacketHandler.cpp b/Engine/Source/Runtime/PacketHandlers/PacketHandler/Private/PacketHandler.cpp
--- a/Engine/Source/Runtime/PacketHandlers/PacketHandler/Private/PacketHandler.cpp
+++ b/Engine/Source/Runtime/PacketHandlers/PacketHandler/Private/PacketHandler.cpp
@@ -1,7 +1,7 @@
 // Copyright 1998-2019 Epic Games, Inc. All Rights Reserved.
 
 #include "PacketHandler.h"
-#include "PacketAudit.h"
+#include "Net/Core/Misc/PacketAudit.h"
 #include "EncryptionComponent.h"
 
 #include "Misc/ConfigCacheIni.h"
@@ -10,7 +10,7 @@
 #include "UObject/Package.h"
 #include "HAL/ConsoleManager.h"
 
-#include "DDoSDetection.h"
+#include "Net/Core/Misc/DDoSDetection.h"
 #include "HandlerComponentFactory.h"
 #include "ReliabilityHandlerComponent.h"
 #include "PacketHandlerProfileConfig.h"
diff --git a/Engine/Source/Runtime/PacketHandlers/PacketHandler/Public/PacketHandler.h b/Engine/Source/Runtime/PacketHandlers/PacketHandler/Public/PacketHandler.h
--- a/Engine/Source/Runtime/PacketHandlers/PacketHandler/Public/PacketHandler.h
+++ b/Engine/Source/Runtime/PacketHandlers/PacketHandler/Public/PacketHandler.h
@@ -7,7 +7,7 @@
 #include "Serialization/BitWriter.h"
 #include "Modules/ModuleInterface.h"
 #include "Containers/Queue.h"
-#include "PacketTraits.h"
+#include "Net/Core/Misc/PacketTraits.h"
 #include "IPAddress.h"
 
 PACKETHANDLER_API DECLARE_LOG_CATEGORY_EXTERN(PacketHandlerLog, Log, All);
diff --git a/Engine/Source/Runtime/PacketHandlers/ReliabilityHandlerComponent/Private/ReliabilityHandlerComponent.cpp b/Engine/Source/Runtime/PacketHandlers/ReliabilityHandlerComponent/Private/ReliabilityHandlerComponent.cpp
--- a/Engine/Source/Runtime/PacketHandlers/ReliabilityHandlerComponent/Private/ReliabilityHandlerComponent.cpp
+++ b/Engine/Source/Runtime/PacketHandlers/ReliabilityHandlerComponent/Private/ReliabilityHandlerComponent.cpp
@@ -2,7 +2,7 @@
 
 #include "ReliabilityHandlerComponent.h"
 #include "Modules/ModuleManager.h"
-#include "PacketAudit.h"
+#include "Net/Core/Misc/PacketAudit.h"
 #include "HAL/PlatformTime.h"
 
 IMPLEMENT_MODULE(FReliabilityHandlerComponentModuleInterface, ReliabilityHandlerComponent);
diff --git a/Engine/Source/Runtime/PacketHandlers/ReliabilityHandlerComponent/ReliabilityHandlerComponent.Build.cs b/Engine/Source/Runtime/PacketHandlers/ReliabilityHandlerComponent/ReliabilityHandlerComponent.Build.cs
--- a/Engine/Source/Runtime/PacketHandlers/ReliabilityHandlerComponent/ReliabilityHandlerComponent.Build.cs
+++ b/Engine/Source/Runtime/PacketHandlers/ReliabilityHandlerComponent/ReliabilityHandlerComponent.Build.cs
@@ -11,7 +11,8 @@ public class ReliabilityHandlerComponent : ModuleRules
         PublicDependencyModuleNames.AddRange(
             new string[] {
 				"Core",
-                "PacketHandler",
+				"NetCore",
+                "PacketHandler"
             }
         );
     }
