Add remaining NetCore files.

diff --git a/Engine/Source/Runtime/Net/Core/NetCore.Build.cs b/Engine/Source/Runtime/Net/Core/NetCore.Build.cs
new file mode 100644
--- /dev/null
+++ b/Engine/Source/Runtime/Net/Core/NetCore.Build.cs
@@ -0,0 +1,19 @@
+// Copyright 1998-2019 Epic Games, Inc. All Rights Reserved.
+
+using UnrealBuildTool;
+using System.IO;
+
+public class NetCore : ModuleRules
+{
+	public NetCore(ReadOnlyTargetRules Target) : base(Target)
+	{
+		PublicDependencyModuleNames.AddRange
+		(
+			new string[]
+			{
+				"Core",
+				"CoreUObject"
+			}
+		);
+	}
+}
diff --git a/Engine/Source/Runtime/Net/Core/Private/Net/Core/Misc/NetCoreLog.h b/Engine/Source/Runtime/Net/Core/Private/Net/Core/Misc/NetCoreLog.h
new file mode 100644
--- /dev/null
+++ b/Engine/Source/Runtime/Net/Core/Private/Net/Core/Misc/NetCoreLog.h
@@ -0,0 +1,8 @@
+// Copyright 1998-2019 Epic Games, Inc. All Rights Reserved.
+
+#pragma once
+
+#include "CoreMinimal.h"
+
+DECLARE_LOG_CATEGORY_EXTERN(LogNetCore, Log, All);
+
diff --git a/Engine/Source/Runtime/Net/Core/Private/Net/Core/NetCoreModule.cpp b/Engine/Source/Runtime/Net/Core/Private/Net/Core/NetCoreModule.cpp
new file mode 100644
--- /dev/null
+++ b/Engine/Source/Runtime/Net/Core/Private/Net/Core/NetCoreModule.cpp
@@ -0,0 +1,32 @@
+// Copyright 1998-2019 Epic Games, Inc. All Rights Reserved.
+
+#include "CoreMinimal.h"
+#include "Modules/ModuleInterface.h"
+#include "Modules/ModuleManager.h"
+#include "Net/Core/Misc/NetCoreLog.h"
+
+// Force export
+#include "Net/Core/Misc/PacketTraits.h"
+
+DEFINE_LOG_CATEGORY(LogNetCore);
+
+
+class FNetCoreModule : public IModuleInterface
+{
+public:
+	virtual void StartupModule() override
+	{
+	}
+
+	virtual void ShutdownModule() override
+	{
+	}
+
+	virtual bool SupportsDynamicReloading() override
+	{
+		return false;
+	}
+};
+
+
+IMPLEMENT_MODULE(FNetCoreModule, NetCore);
