Add more documentation around PushModel, and remove some unused functions from the API.

None

diff --git a/Engine/Source/Runtime/Net/Core/Private/Net/Core/PushModel/PushModel.cpp b/Engine/Source/Runtime/Net/Core/Private/Net/Core/PushModel/PushModel.cpp
--- a/Engine/Source/Runtime/Net/Core/Private/Net/Core/PushModel/PushModel.cpp
+++ b/Engine/Source/Runtime/Net/Core/Private/Net/Core/PushModel/PushModel.cpp
@@ -14,6 +14,157 @@
 
 namespace UE4PushModelPrivate
 {
+	//! Originally, multiple implementations of FPushModelObjectManagers were tested.
+	//! The way this worked was by implementing the managers here, using Macros or
+	//! Switches in the various API functions, and then changing any necessary helper
+	//! Macros to make sure the correct arguments came through.
+	//! Ideally, the interface will always stay the same, and any unused arguments or
+	//! Returns are just ignored.
+	
+	//! TODO: We should add in a way for NetDrivers to opt out of PushModel.
+	//! Things like the Beacon Net Driver, for exmaple, don't need to care about it.
+	//! Since most things are lazily created, this probably isn't a big deal, but
+	//! having explicit behavior preventing it is probably worthwhile.
+
+	/**
+	 * Class that is used to manage Push Model Object states by using Custom IDs.
+	 *
+	 * This class relies on every Replicated UObject to have a Push Model Controlled ID
+	 * associated with it. This ID is injected into the Base Most Replicated class
+	 * during header generation.
+	 *
+	 * For example, UObject does not have the ID field added, but UActorComponent does.
+	 * Since USceneComponent derives from UActorComponent, it can just use the injected
+	 * ID from UActorComponent.
+	 *
+	 * Currently, this works fine for BP because you can't directly derive BPs from UObjects.
+	 * They need to be either AActor or UActorComponent, and those will have the injected IDs.
+	 * If that changes, we may need to move the custom IDs to base UObject and eat the memory.
+	 *
+	 * Also, instead of tracking Objects by pointer, we need another ID that can be used to track
+	 * objects **outside** of networking contexts. For now, we rely on UObject::GetUniqueId, which
+	 * is ultimately the GUObjectArray index of the Object. This same ID is used for Lazy Object Pointers,
+	 * and should be fine so long as that never changes during an Object's lifetime.
+	 *
+	 * **********************************************************************************
+	 * ************* Adding Objects to and Remove Objects From the Manager **************
+	 * **********************************************************************************
+	 *
+	 * Objects are lazily added to the manager (this will be true of any manager).
+	 * Currently, this is tied directly to FRepChangelistState.
+	 * Whenever a FRepChangelistState is created, we will call AddPushModelObject
+	 * and whenever a FRepChangelistState is destroyed, we will call RemovePushModelObject.
+	 *
+	 * FRepChangelistState are only created once on a NetDriver for a given Object in game,
+	 * and are then shared across all connections.
+	 *
+	 * This also means that multiple Net Drivers can "Add" or "Remove" push model objects.
+	 *
+	 * To cope with this, there will is one global FPushModelPerObjectState that is alive
+	 * as long as **any** Net Drivers are replicating the object. There will also be one
+	 * FPushModelPerNetDriverState alive for each individual Net Driver that is replicating
+	 * an object.
+	 *
+	 * **********************************************************************************
+	 * ****************************** Managing Dirty State ******************************
+	 * **********************************************************************************
+	 *
+	 * When an Object is marked dirty from "Game Land", user's will pass a Object's
+	 * Push Model ID as well as the Rep Index for the property (or properties). These
+	 * things can mostly be determined at compile time using the MARK_PROPERTY_DIRTY_* macros,
+	 * meaning that users only really need to pass the Property's Owning Class, the Property's
+	 * name, and an Object Pointer.
+	 *
+	 * If the object is not currently replicated by any net drivers, its Push Model ID will be
+	 * INDEX_NONE, and we will ignore it. Otherwise, we will set the necessary dirty bit
+	 * on the FPushModelPerObjectState.
+	 *
+	 * That dirty state will remain set until we go to replicate the object.
+	 *
+	 * **********************************************************************************
+	 * ************************* Replication using Dirty State **************************
+	 * **********************************************************************************
+	 *
+	 * When a given Net Driver goes to Compare an Object's properties for replication,
+	 * it will request the Push Model State (GetPerNetDriverState). At that point, we will
+	 * push the Global Dirty State (on FPushModelPerObjectState) to the Per Net Driver
+	 * Dirty States (on FPushModelPerNetDriverState). In this case, "Pushing State" is
+	 * effectively a bitwise or of the Push Model Dirty States into the
+	 * Per Net Driver Dirty State, and then clearing the Global Dirty State.
+	 *
+	 * There were 2 main reasons why this was done lazily:
+	 *	1. The subset of Actors that we *will* replicate is always going to be much much
+	 *		smaller than the total number of Actors that *may* replicate.
+	 *
+	 *	2. There isn't a good catch all spot to do this.
+	 *		- Originally having the NetDriver's call this before replication occurred
+	 *			was tested, but properties can be changed during AActor::PreReplication.
+	 *			Further, because Reliable RPCs can force new Channel's open, and initial
+	 *			replication, extra steps had to be taken to try and update states there
+	 *			as well.
+	 *
+	 *		- Using AActor::CalPreReplication was also considered. However, there's
+	 *			currently no good interface for getting all replicated subobjects of an
+	 *			Actor (or other subobjects), so there's no guarantee we'd push states.
+	 *			Additionally, even if there was an interface this wouldn't necessarily
+	 *			respect Subobject Replication Keys.
+	 *
+	 * Doing this lazily is a simple solution to both of these problems, with the trade
+	 * off that if multiple NetDrivers are replicating on the same frame, we may try to
+	 * push the state multiple times. But, we can detect that no properties have changed
+	 * since the last push, and not waste much time.
+	 *
+	 * When a property is *not* marked dirty but is Push Enabled, then we will skip comparing it.
+	 * When a property is marked dirty, or is not Push Enabled, we will compare it.
+	 * See the comments in PushModel.h for more info on that.
+	 *
+	 * **********************************************************************************
+	 * ********************** Potential for "Automatic Dormancy" ************************
+	 * **********************************************************************************
+	 *
+	 * Push model may support the notion of Automatic Dormancy in the future. The premise
+	 * being that instead of relying on designers or developers to to call FlushNetDormancy
+	 * when changing object properties, instead Push Model could be used as a drop in
+	 * substitute. If you're already adhering to the contract of Push Model, your system
+	 * will already have all the necessary hooks to determine when properties change.
+	 *
+	 * This comes with a few nice benefits:
+	 *	1. It would be impossible to forget to wake an Actor up when properties change.
+	 *
+	 *	2. Similarly, if you Components / Subobjects that contained shared logic,
+	 *		there would be no extra special casing to mark the owning Actor awake.
+	 *
+	 *	3. **Every** type of object that fully relied on Push Model could automatically
+	 *		make use of this system without the need for extra configuration.
+	 *		This is especially useful considering all Blueprints derived directly from
+	 *		AActor, UActorComponent, and USceneComponent would currently be in this
+	 *		category. The more engine conversion that happens, the broader that becomes.
+	 *
+	 * Speculative plan for how this might work:
+	 *
+	 *	1. The fundamental machinery of Dormancy stays the same.
+	 *		An Object is eligible for Dormancy if none of it's properties have changed
+	 *		for some configurable timeout, and all connections have received its most
+	 *		up to date information.
+	 *
+	 *	2. Alongisde dirty property states, Push Model could also have a bitfield that
+	 *		tracks whether or not an object was dirtied in a frame. Alternatively,
+	 *		that state could be derived from FPushModelPerObjectState.
+	 *
+	 *	3. At some defined point in a frame before Replication of any NetDrivers
+	 *		occurs, we push the Dirty Object State to the Net Driver.
+	 *		This might work by just having each Net Driver grab the Dirty Object
+	 *		bitfield and if a given object is dirty, mapping the Object ID back
+	 *		to Actor / Object, and then calling the normal Flush Net Dormancy calls.
+	 *
+	 *	4. At some defined point in a frame *after* Replication of *every* NetDriver,
+	 *		we reset the Object Dirty State.
+	 *
+	 * The trickiest part of this would likely be finding an efficient mapping back from a
+	 * PushModel ID to a Networked Object. Push Model was designed to be mostly agnostic
+	 * to Objects and because of that explicitly doesn't provide any API to go from a
+	 * Push ID back to an Object.
+	 */
 	class FPushModelObjectManager_CustomId
 	{
 	private:
@@ -82,14 +233,6 @@ namespace UE4PushModelPrivate
 			}
 		}
 
-		void PreReplication()
-		{
-			for (FPushModelPerObjectState& PerObjectState : PerObjectStates)
-			{
-				PerObjectState.PushDirtyStateToNetDrivers();
-			}
-		}
-
 		void PostGarbageCollect()
 		{
 			// We can't compact PerObjectStates because we need ObjectIDs to be stable.
diff --git a/Engine/Source/Runtime/Net/Core/Public/Net/Core/PushModel/PushModel.h b/Engine/Source/Runtime/Net/Core/Public/Net/Core/PushModel/PushModel.h
--- a/Engine/Source/Runtime/Net/Core/Public/Net/Core/PushModel/PushModel.h
+++ b/Engine/Source/Runtime/Net/Core/Public/Net/Core/PushModel/PushModel.h
@@ -80,6 +80,9 @@
  * but if they are using things like Blueprint Setters, passing values around by reference, or setting values in Native
  * C++, then it is completely up to them to call the necessary methods.
  *
+ * UFUNCTIONS or Blueprint Defined Functions that take Properties by reference will check for Networked Properties and
+ * should add nodes automatically (see UNetPushModelHelpers).
+ *
  * Devs don't need to care *if* a given object is actually networked or is actively replicating, Push Model handles
  * that sort of tracking internally. As such, we don't expect users to ever need to know or have a reference to
  * any information other than A) The Object that's being dirtied and B) The Property that's being dirtied.
@@ -95,11 +98,15 @@
  * whether or not a given property is dirty. If it is, then we can compare and replicate it normally, otherwise we
  * can quickly skip it.
  *
- * One important thing to note is that at this time, Push Model only supports Top Level properties of Objects.
+ * One important thing to note is that at this time, Push Model only supports Top Level properties of Objects
+ * (in RepLayout parlance, Parent Commands).
+ *
  * So, if changes are made to Structs, Containers, or nested properties therein, the owning Struct or Container
  * property must be marked dirty.
  *
  * This also means that we *must* compare a property when it's dirty to make sure we know exactly what's changed.
+ * Optimizations could be made for "network primtiive" types where we skip comparisons if desired, at the potential
+ * expense of extra bandwidth for properties that didn't actually change.
  *
  ****************************************************************
  * Examples
@@ -333,12 +340,6 @@ namespace UE4PushModelPrivate
 	 */
 	NETCORE_API void MarkPropertyDirty(const FNetPushObjectId ObjectId, const int32 StartRepIndex, const int32 EndRepIndex);
 
-	/**
-	 * This method is called by all NetDriver's before they replicated properties.
-	 * It's required in order to give PushModel a chance to do any setup work.
-	 */
-	NETCORE_API void PreReplication();
-	
 	/**
 	 * This method is called by all NetDriver's PostGarbageCollection.
 	 * It's required in order to give PushModel a chance to do any cleanup work.
