Fix push model compilation for linux and android

diff --git a/Engine/Source/Runtime/Net/Core/Public/Net/Core/PushModel/PushModel.h b/Engine/Source/Runtime/Net/Core/Public/Net/Core/PushModel/PushModel.h
--- a/Engine/Source/Runtime/Net/Core/Public/Net/Core/PushModel/PushModel.h
+++ b/Engine/Source/Runtime/Net/Core/Public/Net/Core/PushModel/PushModel.h
@@ -360,9 +360,9 @@ namespace UE4PushModelPrivate
 #define PUSH_MAKE_BP_PROPERTIES_PUSH_MODEL() (UE4PushModelPrivate::IsPushModelEnabled() && UE4PushModelPrivate::MakeBpPropertiesPushModel())
 
 #define GET_PROPERTY_REP_INDEX(ClassName, PropertyName) (int32)ClassName::ENetFields_Private::PropertyName
-#define GET_PROPERTY_REP_INDEX_STATIC_ARRAY_START(ClassName, PropertyName) ((int32)GET_PROPERTY_REP_INDEX(ClassName, PropertyName)_STATIC_ARRAY)
-#define GET_PROPERTY_REP_INDEX_STATIC_ARRAY_END(ClassName, PropertyName) ((int32)GET_PROPERTY_REP_INDEX(ClassName, PropertyName)_STATIC_ARRAY_END)
-#define GET_PROPERTY_REP_INDEX_STATIC_ARRAY_INDEX(ClassName, PropertyName, ArrayIndex) ((int32)GET_PROPERTY_REP_INDEX_STATIC_ARRAY_START(ClassName, PropertyName) + ArrayIndex)
+#define GET_PROPERTY_REP_INDEX_STATIC_ARRAY_START(ClassName, PropertyName) ((int32)ClassName::ENetFields_Private::PropertyName ## _STATIC_ARRAY)
+#define GET_PROPERTY_REP_INDEX_STATIC_ARRAY_END(ClassName, PropertyName) ((int32)ClassName::ENetFields_Private::PropertyName ## _STATIC_ARRAY_END)
+#define GET_PROPERTY_REP_INDEX_STATIC_ARRAY_INDEX(ClassName, PropertyName, ArrayIndex) (GET_PROPERTY_REP_INDEX_STATIC_ARRAY_START(ClassName, PropertyName) + ArrayIndex)
 
 #define IS_PROPERTY_REPLICATED(Property) (0 != (EPropertyFlags::CPF_Net & Property->PropertyFlags))
 
