Updated -NetTrace to work with new channel concept for trace system.
Specifying -NetTrace=[VerbosityLevel] on the commanline implicitly enabled the NetChannel TraceChannel

diff --git a/Engine/Source/Runtime/Net/Core/Private/Net/Core/Trace/NetTraceInternal.cpp b/Engine/Source/Runtime/Net/Core/Private/Net/Core/Trace/NetTraceInternal.cpp
--- a/Engine/Source/Runtime/Net/Core/Private/Net/Core/Trace/NetTraceInternal.cpp
+++ b/Engine/Source/Runtime/Net/Core/Private/Net/Core/Trace/NetTraceInternal.cpp
@@ -7,6 +7,7 @@
 #include "HAL/PlatformTLS.h"
 #include "HAL/TlsAutoCleanup.h"
 #include "Hash/CityHash.h"
+#include "Trace/Trace.h"	
 #include "UObject/NameTypes.h"
 
 #define UE_NET_TRACE_VERSION ENetTraceVersion_BunchChannelIndex
@@ -73,12 +74,16 @@ void FNetTrace::SetTraceVerbosity(uint32 Verbosity)
 	// Enable
 	if (!GetTraceVerbosity() && NewVerbosity)
 	{
+		Trace::ToggleChannel(TEXT("NetChannel"), true);
+
 		FNetTraceInternal::Reporter::ReportInitEvent(FNetTraceInternal::NetTraceVersion);
 	}
 	else if (GetTraceVerbosity() && !NewVerbosity)
 	{
 		if (FNetTraceInternal::ThreadBuffer)
 		{
+			Trace::ToggleChannel(TEXT("NetChannel"), false);
+
 			delete FNetTraceInternal::ThreadBuffer;
 			FNetTraceInternal::ThreadBuffer = nullptr;
 		}		
