Make GetNetPushIdDynamic public

#rnx

diff --git a/Engine/Source/Runtime/CoreUObject/Public/UObject/Object.h b/Engine/Source/Runtime/CoreUObject/Public/UObject/Object.h
--- a/Engine/Source/Runtime/CoreUObject/Public/UObject/Object.h
+++ b/Engine/Source/Runtime/CoreUObject/Public/UObject/Object.h
@@ -1481,7 +1481,10 @@ private:
 		// always have this implemented (by UHT).
 		check(false);
 	}
-	
+
+public:
+
+	/** Should only ever be used by internal systems. */
 	virtual int32 GetNetPushIdDynamic() const
 	{
 		return INDEX_NONE;
@@ -1492,18 +1495,11 @@ struct FObjectNetPushIdHelper
 {
 private:
 	friend struct FNetPrivatePushIdHelper;
-	friend class UNetPushModelHelpers;
-	friend class UKismetArrayLibrary;
 
 	static void SetNetPushIdDynamic(UObject* Object, const int32 NewNetPushId)
 	{
 		Object->SetNetPushIdDynamic(NewNetPushId);
 	}
-	
-	static int32 GetNetPushIdDynamic(const UObject* const Object)
-	{
-		return Object->GetNetPushIdDynamic();
-	}
 };
 
 /**
diff --git a/Engine/Source/Runtime/Engine/Private/RepLayout.cpp b/Engine/Source/Runtime/Engine/Private/RepLayout.cpp
--- a/Engine/Source/Runtime/Engine/Private/RepLayout.cpp
+++ b/Engine/Source/Runtime/Engine/Private/RepLayout.cpp
@@ -780,7 +780,7 @@ struct FNetPrivatePushIdHelper
 {
 	static void SetNetPushID(UObject* InObject, const UE4PushModelPrivate::FNetPushObjectId ObjectId)
 	{
-		const UE4PushModelPrivate::FNetPushObjectId CurrentId = FObjectNetPushIdHelper::GetNetPushIdDynamic(InObject);
+		const UE4PushModelPrivate::FNetPushObjectId CurrentId = InObject->GetNetPushIdDynamic();
 		if (CurrentId != ObjectId)
 		{
 			check(CurrentId == INDEX_NONE);
diff --git a/Engine/Source/Runtime/Net/Core/Public/Net/Core/PushModel/PushModel.h b/Engine/Source/Runtime/Net/Core/Public/Net/Core/PushModel/PushModel.h
--- a/Engine/Source/Runtime/Net/Core/Public/Net/Core/PushModel/PushModel.h
+++ b/Engine/Source/Runtime/Net/Core/Public/Net/Core/PushModel/PushModel.h
@@ -367,8 +367,8 @@ namespace UE4PushModelPrivate
 #define IS_PROPERTY_REPLICATED(Property) (0 != (EPropertyFlags::CPF_Net & Property->PropertyFlags))
 
 #define CONDITIONAL_ON_OBJECT_NET_ID(Object, Work) { const UE4PushModelPrivate::FNetPushObjectId PrivatePushId = Object->GetNetPushId(); if (INDEX_NONE != PrivatePushId) { Work; } }
-#define CONDITIONAL_ON_OBJECT_NET_ID_DYNAMIC(Object, Work) { const UE4PushModelPrivate::FNetPushObjectId PrivatePushId = FObjectNetPushIdHelper::GetNetPushIdDynamic(Object); if (INDEX_NONE != PrivatePushId) { Work; } }
-#define CONDITIONAL_ON_REP_INDEX_AND_OBJECT_NET_ID(Object, Property, Work) if (IS_PROPERTY_REPLICATED(Property)) { const UE4PushModelPrivate::FNetPushObjectId PrivatePushId = FObjectNetPushIdHelper::GetNetPushIdDynamic(Object); if (PrivatePushId != INDEX_NONE) { Work; } }
+#define CONDITIONAL_ON_OBJECT_NET_ID_DYNAMIC(Object, Work) { const UE4PushModelPrivate::FNetPushObjectId PrivatePushId = Object->GetNetPushIdDynamic(); if (INDEX_NONE != PrivatePushId) { Work; } }
+#define CONDITIONAL_ON_REP_INDEX_AND_OBJECT_NET_ID(Object, Property, Work) if (IS_PROPERTY_REPLICATED(Property)) { const UE4PushModelPrivate::FNetPushObjectId PrivatePushId = Object->GetNetPushIdDynamic(); if (PrivatePushId != INDEX_NONE) { Work; } }
 
 //~ For these macros, we won't bother checking if Push Model is enabled. Instead, we'll just check to see whether or not the Custom ID is valid.
 
diff --git a/Engine/Source/Runtime/Net/Core/Public/Net/Core/PushModel/PushModelMacros.h b/Engine/Source/Runtime/Net/Core/Public/Net/Core/PushModel/PushModelMacros.h
--- a/Engine/Source/Runtime/Net/Core/Public/Net/Core/PushModel/PushModelMacros.h
+++ b/Engine/Source/Runtime/Net/Core/Public/Net/Core/PushModel/PushModelMacros.h
@@ -16,9 +16,9 @@
 #define REPLICATED_BASE_CLASS(ClassName) \
 private: \
 int32 NetPushId_Internal = INDEX_NONE; \
-virtual void SetNetPushIdDynamic(const int32 InNetPushId) override { NetPushId_Internal = InNetPushId; } \
-virtual int32 GetNetPushIdDynamic() const override { return GetNetPushId(); } \
+virtual void SetNetPushIdDynamic(const int32 InNetPushId) override final { NetPushId_Internal = InNetPushId; } \
 public: \
+virtual int32 GetNetPushIdDynamic() const override final { return GetNetPushId(); } \
 int32 GetNetPushId() const { return NetPushId_Internal; }
 
 #else
