Prevent PushModel from accessing internal memory for objects that are no longer replicating.

#rnx

diff --git a/Engine/Source/Runtime/Engine/Private/NetworkDriver.cpp b/Engine/Source/Runtime/Engine/Private/NetworkDriver.cpp
--- a/Engine/Source/Runtime/Engine/Private/NetworkDriver.cpp
+++ b/Engine/Source/Runtime/Engine/Private/NetworkDriver.cpp
@@ -3334,8 +3334,6 @@ void UNetDriver::PostGarbageCollect()
 			It.RemoveCurrent();
 		}
 	}
-	
-	CALL_PUSH_MODEL_POSTGARBAGECOLLECT();
 }
 
 void UNetDriver::AddReferencedObjects(UObject* InThis, FReferenceCollector& Collector)
diff --git a/Engine/Source/Runtime/Net/Core/Private/Net/Core/PushModel/PushModel.cpp b/Engine/Source/Runtime/Net/Core/Private/Net/Core/PushModel/PushModel.cpp
--- a/Engine/Source/Runtime/Net/Core/Private/Net/Core/PushModel/PushModel.cpp
+++ b/Engine/Source/Runtime/Net/Core/Private/Net/Core/PushModel/PushModel.cpp
@@ -10,6 +10,7 @@
 #include "Types/PushModelUtils.h"
 #include "Types/PushModelPerObjectState.h"
 #include "Types/PushModelPerNetDriverState.h"
+#include "UObject/UObjectGlobals.h"
 
 
 namespace UE4PushModelPrivate
@@ -175,24 +176,42 @@ namespace UE4PushModelPrivate
 
 		FPushModelObjectManager_CustomId()
 		{
+			PostGarbageCollectHandle = FCoreUObjectDelegates::GetPostGarbageCollect().AddRaw(this, &FPushModelObjectManager_CustomId::PostGarbageCollect);
 		}
 
 		~FPushModelObjectManager_CustomId()
 		{
+			FCoreUObjectDelegates::GetPostGarbageCollect().Remove(PostGarbageCollectHandle);
 		}
 
 		void MarkPropertyDirty(const FNetPushObjectId ObjectId, const int32 RepIndex)
 		{
-			// The macros will take care of filtering out invalid objects, so we don't need to check here.
-			PerObjectStates[ObjectId].MarkPropertyDirty(RepIndex);
+			const int32 ObjectIndex = ObjectId;
+			if (LIKELY(PerObjectStates.IsValidIndex(ObjectIndex)))
+			{
+				// The macros will take care of filtering out invalid objects, so we don't need to check here.
+				PerObjectStates[ObjectIndex].MarkPropertyDirty(RepIndex);
+			}
+			else
+			{
+				ensure(false);
+			}
 		}
 
 		void MarkPropertyDirty(const FNetPushObjectId ObjectId, const int32 StartRepIndex, const int32 EndRepIndex)
 		{
-			FPushModelPerObjectState& ObjectState = PerObjectStates[ObjectId];
-			for (int RepIndex = StartRepIndex; RepIndex <= EndRepIndex; ++RepIndex)
+			const int32 ObjectIndex = ObjectId;
+			if (LIKELY(PerObjectStates.IsValidIndex(ObjectIndex)))
+			{
+				FPushModelPerObjectState& ObjectState = PerObjectStates[ObjectIndex];
+				for (int RepIndex = StartRepIndex; RepIndex <= EndRepIndex; ++RepIndex)
+				{
+					ObjectState.MarkPropertyDirty(RepIndex);
+				}
+			}
+			else
 			{
-				ObjectState.MarkPropertyDirty(RepIndex);
+				ensure(false);
 			}
 		}
 
@@ -216,20 +235,11 @@ namespace UE4PushModelPrivate
 
 		void RemoveNetworkObject(const FPushModelPerNetDriverHandle Handle)
 		{
-			if (PerObjectStates.IsValidIndex(Handle.ObjectId))
+			const int32 ObjectIndex = Handle.ObjectId;
+			if (LIKELY(PerObjectStates.IsValidIndex(ObjectIndex)))
 			{
-				FPushModelPerObjectState& PerObjectState = PerObjectStates[Handle.ObjectId];
+				FPushModelPerObjectState& PerObjectState = PerObjectStates[ObjectIndex];
 				PerObjectState.RemovePerNetDriverState(Handle.NetDriverId);
-				if (!PerObjectState.HasAnyNetDriverStates())
-				{
-					ObjectIdToInternalId.Remove(PerObjectState.GetObjectId());
-					PerObjectStates.RemoveAt(Handle.ObjectId);
-
-					if (NewObjectLookupPosition > Handle.ObjectId)
-					{
-						NewObjectLookupPosition = Handle.ObjectId;
-					}
-				}
 			}
 		}
 
@@ -238,6 +248,31 @@ namespace UE4PushModelPrivate
 			// We can't compact PerObjectStates because we need ObjectIDs to be stable.
 			// But we can shrink it.
 
+			// Go ahead and remove any PerObjectStates that aren't being tracked by any NetDrivers.
+			// We have to wait until GC for this, because the NetDrivers will periodically remove
+			// Network Objects that are still alive (but marked Pending Kill) but we don't have a way
+			// to safely clear the Push Model Handles from those objects.
+			//
+			// That means if we tried to remove these items from Push Model tracking, we could end up
+			// with cases where we reassign the Push Model ID to a new object, and the old object could
+			// inadvertently dirty its state.
+			//
+			// In theory, this should never happen because once the object is marked Pending Kill none
+			// of its properties should change again, but it's also possible that calls like BeginDestroy
+			// could modify properties, etc.
+			//
+			// Currently, none of these objects are actually removed though unless the networking system
+			// detects they are PendingKill (their WeakObjectPtr can't be resolved anymore), so there shouldn't
+			// be any cases where we remove these for "still alive" objects.
+			for (auto It = PerObjectStates.CreateIterator(); It; ++It)
+			{
+				if (!It->HasAnyNetDriverStates())
+				{
+					ObjectIdToInternalId.Remove(It->GetObjectId());
+					It.RemoveCurrent();
+				}
+			}
+
 			PerObjectStates.Shrink();
 			ObjectIdToInternalId.Compact();
 			NewObjectLookupPosition = 0;
@@ -245,9 +280,10 @@ namespace UE4PushModelPrivate
 
 		FPushModelPerNetDriverState* GetPerNetDriverState(const FPushModelPerNetDriverHandle Handle)
 		{
-			if (PerObjectStates.IsValidIndex(Handle.ObjectId))
+			const int32 ObjectIndex = Handle.ObjectId;
+			if (LIKELY(PerObjectStates.IsValidIndex(ObjectIndex)))
 			{
-				FPushModelPerObjectState& ObjectState = PerObjectStates[Handle.ObjectId];
+				FPushModelPerObjectState& ObjectState = PerObjectStates[ObjectIndex];
 				ObjectState.PushDirtyStateToNetDrivers();
 				return &ObjectState.GetPerNetDriverState(Handle.NetDriverId);
 			}
@@ -260,6 +296,7 @@ namespace UE4PushModelPrivate
 		int32 NewObjectLookupPosition = 0;
 		TMap<FNetPushObjectId, FNetPushObjectId> ObjectIdToInternalId;
 		TSparseArray<FPushModelPerObjectState> PerObjectStates;
+		FDelegateHandle PostGarbageCollectHandle;
 	};
 
 	static FPushModelObjectManager_CustomId PushObjectManager;
@@ -288,11 +325,6 @@ namespace UE4PushModelPrivate
 		PushObjectManager.MarkPropertyDirty(ObjectId, StartRepIndex, EndRepIndex);
 	}
 
-	void PostGarbageCollect()
-	{
-		PushObjectManager.PostGarbageCollect();
-	}
-
 	/**
 	 * Called by a given NetDriver to notify us that it's seen a given Object for the first time (or the first time
 	 * since it was removed).
diff --git a/Engine/Source/Runtime/Net/Core/Public/Net/Core/PushModel/PushModel.h b/Engine/Source/Runtime/Net/Core/Public/Net/Core/PushModel/PushModel.h
--- a/Engine/Source/Runtime/Net/Core/Public/Net/Core/PushModel/PushModel.h
+++ b/Engine/Source/Runtime/Net/Core/Public/Net/Core/PushModel/PushModel.h
@@ -340,12 +340,7 @@ namespace UE4PushModelPrivate
 	 */
 	NETCORE_API void MarkPropertyDirty(const FNetPushObjectId ObjectId, const int32 StartRepIndex, const int32 EndRepIndex);
 
-	/**
-	 * This method is called by all NetDriver's PostGarbageCollection.
-	 * It's required in order to give PushModel a chance to do any cleanup work.
-	 */
-	NETCORE_API void PostGarbageCollect();
-	
+
 	NETCORE_API const FPushModelPerNetDriverHandle AddPushModelObject(const FNetPushObjectId ObjectId, const uint16 NumberOfReplicatedProperties);
 	NETCORE_API void RemovePushModelObject(const FPushModelPerNetDriverHandle Handle);
 
@@ -354,8 +349,6 @@ namespace UE4PushModelPrivate
 
 
 #define CONDITIONAL_ON_PUSH_MODEL(Work) if (UE4PushModelPrivate::IsPushModelEnabled()) { Work; }
-#define CALL_PUSH_MODEL_PREREPLICATION() CONDITIONAL_ON_PUSH_MODEL(UE4PushModelPrivate::PreReplication();)
-#define CALL_PUSH_MODEL_POSTGARBAGECOLLECT() CONDITIONAL_ON_PUSH_MODEL(UE4PushModelPrivate::PostGarbageCollect();)
 #define IS_PUSH_MODEL_ENABLED() UE4PushModelPrivate::IsPushModelEnabled()
 #define PUSH_MAKE_BP_PROPERTIES_PUSH_MODEL() (UE4PushModelPrivate::IsPushModelEnabled() && UE4PushModelPrivate::MakeBpPropertiesPushModel())
 
@@ -366,9 +359,9 @@ namespace UE4PushModelPrivate
 
 #define IS_PROPERTY_REPLICATED(Property) (0 != (EPropertyFlags::CPF_Net & Property->PropertyFlags))
 
-#define CONDITIONAL_ON_OBJECT_NET_ID(Object, Work) { const UE4PushModelPrivate::FNetPushObjectId PrivatePushId = Object->GetNetPushId(); if (INDEX_NONE != PrivatePushId) { Work; } }
-#define CONDITIONAL_ON_OBJECT_NET_ID_DYNAMIC(Object, Work) { const UE4PushModelPrivate::FNetPushObjectId PrivatePushId = Object->GetNetPushIdDynamic(); if (INDEX_NONE != PrivatePushId) { Work; } }
-#define CONDITIONAL_ON_REP_INDEX_AND_OBJECT_NET_ID(Object, Property, Work) if (IS_PROPERTY_REPLICATED(Property)) { const UE4PushModelPrivate::FNetPushObjectId PrivatePushId = Object->GetNetPushIdDynamic(); if (PrivatePushId != INDEX_NONE) { Work; } }
+#define CONDITIONAL_ON_OBJECT_NET_ID(Object, Work) { const UE4PushModelPrivate::FNetPushObjectId PrivatePushId = Object->GetNetPushId(); Work; }
+#define CONDITIONAL_ON_OBJECT_NET_ID_DYNAMIC(Object, Work) { const UE4PushModelPrivate::FNetPushObjectId PrivatePushId = Object->GetNetPushIdDynamic(); Work; }
+#define CONDITIONAL_ON_REP_INDEX_AND_OBJECT_NET_ID(Object, Property, Work) if (IS_PROPERTY_REPLICATED(Property)) { const UE4PushModelPrivate::FNetPushObjectId PrivatePushId = Object->GetNetPushIdDynamic();  Work; }
 
 //~ For these macros, we won't bother checking if Push Model is enabled. Instead, we'll just check to see whether or not the Custom ID is valid.
 
@@ -413,7 +406,4 @@ namespace UE4PushModelPrivate
 #define IS_PUSH_MODEL_ENABLED() false
 #define PUSH_MAKE_BP_PROPERTIES_PUSH_MODEL() false
 
-#define CALL_PUSH_MODEL_PREREPLICATION() 
-#define CALL_PUSH_MODEL_POSTGARBAGECOLLECT() 
-
 #endif
\ No newline at end of file
