Fix missing file

#rnx

diff --git a/Engine/Source/Runtime/Net/Core/Private/Net/Core/PushModel/Types/PushModelPerNetDriverState.h b/Engine/Source/Runtime/Net/Core/Private/Net/Core/PushModel/Types/PushModelPerNetDriverState.h
--- a/Engine/Source/Runtime/Net/Core/Private/Net/Core/PushModel/Types/PushModelPerNetDriverState.h
+++ b/Engine/Source/Runtime/Net/Core/Private/Net/Core/PushModel/Types/PushModelPerNetDriverState.h
@@ -18,20 +18,28 @@ namespace UE4PushModelPrivate
 
 		FPushModelPerNetDriverState(const uint16 InNumberOfProperties)
 			: PropertyDirtyStates(true, InNumberOfProperties)
+			, bRecentlyCollectedGarbage(false)
 		{
 		}
 		
 		FPushModelPerNetDriverState(FPushModelPerNetDriverState&& Other)
 			: PropertyDirtyStates(MoveTemp(Other.PropertyDirtyStates))
+			, bRecentlyCollectedGarbage(Other.bRecentlyCollectedGarbage)
 		{
 		}
 		
 		FPushModelPerNetDriverState(const FPushModelPerNetDriverState& Other) = delete;
 		FPushModelPerNetDriverState& operator=(const FPushModelPerNetDriverState& Other) = delete;
 
+		void SetRecentlyCollectedGarbage()
+		{
+			bRecentlyCollectedGarbage = true;
+		}
+
 		void ResetDirtyStates()
 		{
 			ResetBitArray(PropertyDirtyStates);
+			bRecentlyCollectedGarbage = false;
 		}
 
 		void CountBytes(FArchive& Ar) const
@@ -39,7 +47,7 @@ namespace UE4PushModelPrivate
 			PropertyDirtyStates.CountBytes(Ar);
 		}
 
-		const bool IsPropertyDirty(const uint16 RepIndex)
+		bool IsPropertyDirty(const uint16 RepIndex) const
 		{
 			return PropertyDirtyStates[RepIndex];
 		}
@@ -49,6 +57,11 @@ namespace UE4PushModelPrivate
 			return TConstSetBitIterator<>(PropertyDirtyStates);
 		}
 
+		bool DidRecentlyCollectGarbage() const
+		{
+			return bRecentlyCollectedGarbage;
+		}
+
 		void MarkPropertiesDirty(const TBitArray<>& OtherBitArray)
 		{
 			BitwiseOrBitArrays(OtherBitArray, PropertyDirtyStates);
@@ -66,6 +79,7 @@ namespace UE4PushModelPrivate
 		 * Note, bits will be allocated for all replicated properties, not just push model properties.
 		 */
 		TBitArray<> PropertyDirtyStates;
+		uint8 bRecentlyCollectedGarbage : 1;
 	};
 }
 
