Use FObjectKey in PushModel tracking instead of just UniqueId.
This will guarantee uniqueness across Garbage Collections to prevent potential collisions.

diff --git a/Engine/Source/Runtime/Engine/Private/RepLayout.cpp b/Engine/Source/Runtime/Engine/Private/RepLayout.cpp
--- a/Engine/Source/Runtime/Engine/Private/RepLayout.cpp
+++ b/Engine/Source/Runtime/Engine/Private/RepLayout.cpp
@@ -178,14 +178,15 @@ namespace UE4_RepLayout_Private
 	}
 
 #if WITH_PUSH_MODEL
-	const UE4PushModelPrivate::FPushModelPerNetDriverHandle ConditionallyAddPushModelObject(const UE4PushModelPrivate::FNetPushObjectId ObjectId, const TSharedRef<const FRepLayout>& InRepLayout)
+	const UE4PushModelPrivate::FPushModelPerNetDriverHandle ConditionallyAddPushModelObject(const UObject* const Object, const TSharedRef<const FRepLayout>& InRepLayout)
 	{
 		const int32 NumReplicatedProperties = InRepLayout->GetNumParents();
 		if (UE4PushModelPrivate::IsPushModelEnabled() &&
 			NumReplicatedProperties > 0 &&
 			EnumHasAnyFlags(InRepLayout->GetFlags(), ERepLayoutFlags::FullPushSupport | ERepLayoutFlags::PartialPushSupport))
 		{
-			return UE4PushModelPrivate::AddPushModelObject(ObjectId, NumReplicatedProperties);
+			const FObjectKey ObjectKey(Object);
+			return UE4PushModelPrivate::AddPushModelObject(ObjectKey, NumReplicatedProperties);
 		}
 
 		return UE4PushModelPrivate::FPushModelPerNetDriverHandle::MakeInvalidHandle();
@@ -974,7 +975,7 @@ FRepChangelistState::FRepChangelistState(
 	, StaticBuffer(InRepLayout->CreateShadowBuffer(InSource))
 
 #if WITH_PUSH_MODEL
-	, PushModelObjectHandle(UE4_RepLayout_Private::ConditionallyAddPushModelObject(InRepresenting->GetUniqueID(), InRepLayout))
+	, PushModelObjectHandle(UE4_RepLayout_Private::ConditionallyAddPushModelObject(InRepresenting, InRepLayout))
 #endif // WITH_PUSH_MODEL
 
 {
diff --git a/Engine/Source/Runtime/Net/Core/Private/Net/Core/PushModel/PushModel.cpp b/Engine/Source/Runtime/Net/Core/Private/Net/Core/PushModel/PushModel.cpp
--- a/Engine/Source/Runtime/Net/Core/Private/Net/Core/PushModel/PushModel.cpp
+++ b/Engine/Source/Runtime/Net/Core/Private/Net/Core/PushModel/PushModel.cpp
@@ -209,19 +209,19 @@ namespace UE4PushModelPrivate
 			}
 		}
 
-		const FPushModelPerNetDriverHandle AddNetworkObject(const FNetPushObjectId ObjectId, const uint16 NumReplicatedProperties)
+		const FPushModelPerNetDriverHandle AddNetworkObject(const FObjectKey ObjectKey, const uint16 NumReplicatedProperties)
 		{
-			FNetPushObjectId& InternalPushId = ObjectIdToInternalId.FindOrAdd(ObjectId, INDEX_NONE);
+			FNetPushObjectId& InternalPushId = ObjectKeyToInternalId.FindOrAdd(ObjectKey, INDEX_NONE);
 			if (INDEX_NONE == InternalPushId)
 			{
 				FSparseArrayAllocationInfo AllocationInfo = PerObjectStates.AddUninitializedAtLowestFreeIndex(NewObjectLookupPosition);
-				new (AllocationInfo.Pointer) FPushModelPerObjectState(ObjectId, NumReplicatedProperties);
+				new (AllocationInfo.Pointer) FPushModelPerObjectState(ObjectKey, NumReplicatedProperties);
 				InternalPushId = AllocationInfo.Index;
 			}
 
 			FPushModelPerObjectState& PerObjectState = PerObjectStates[InternalPushId];
 			check(PerObjectState.GetNumberOfProperties() == NumReplicatedProperties);
-			check(PerObjectState.GetObjectId() == ObjectId);
+			check(PerObjectState.GetObjectKey() == ObjectKey);
 
 			const FNetPushPerNetDriverId NetDriverId = PerObjectState.AddPerNetDriverState();
 			return FPushModelPerNetDriverHandle(NetDriverId, InternalPushId);
@@ -264,7 +264,7 @@ namespace UE4PushModelPrivate
 			{
 				if (!It->HasAnyNetDriverStates())
 				{
-					ObjectIdToInternalId.Remove(It->GetObjectId());
+					ObjectKeyToInternalId.Remove(It->GetObjectKey());
 					It.RemoveCurrent();
 				}
 				else
@@ -274,7 +274,7 @@ namespace UE4PushModelPrivate
 			}
 
 			PerObjectStates.Shrink();
-			ObjectIdToInternalId.Compact();
+			ObjectKeyToInternalId.Compact();
 			NewObjectLookupPosition = 0;
 		}
 
@@ -294,7 +294,7 @@ namespace UE4PushModelPrivate
 	private:
 
 		int32 NewObjectLookupPosition = 0;
-		TMap<FNetPushObjectId, FNetPushObjectId> ObjectIdToInternalId;
+		TMap<FObjectKey, FNetPushObjectId> ObjectKeyToInternalId;
 		TSparseArray<FPushModelPerObjectState> PerObjectStates;
 		FDelegateHandle PostGarbageCollectHandle;
 	};
@@ -332,12 +332,12 @@ namespace UE4PushModelPrivate
 	 * This may be called multiple times for a given Object if there are multiple NetDrivers, but it's expected
 	 * that each NetDriver only calls this once per object before RemoteNetworkObject is called.
 	 *
-	 * @param ObjectId						The UniqueId for the object.
+	 * @param ObjectKey						An ObjectKey to uniquely identify the object.
 	 * @param NumberOfReplicatedProperties	The number of replicated properties for this object.
 	 *
 	 * @return A Handle that can be used in other calls to uniquely identify this object per NetDriver.
 	 */
-	const FPushModelPerNetDriverHandle AddPushModelObject(const FNetPushObjectId ObjectId, const uint16 NumberOfReplicatedProperties)
+	const FPushModelPerNetDriverHandle AddPushModelObject(const FObjectKey ObjectId, const uint16 NumberOfReplicatedProperties)
 	{
 		return PushObjectManager.AddNetworkObject(ObjectId, NumberOfReplicatedProperties);
 	}
diff --git a/Engine/Source/Runtime/Net/Core/Private/Net/Core/PushModel/Types/PushModelPerObjectState.h b/Engine/Source/Runtime/Net/Core/Private/Net/Core/PushModel/Types/PushModelPerObjectState.h
--- a/Engine/Source/Runtime/Net/Core/Private/Net/Core/PushModel/Types/PushModelPerObjectState.h
+++ b/Engine/Source/Runtime/Net/Core/Private/Net/Core/PushModel/Types/PushModelPerObjectState.h
@@ -11,6 +11,7 @@
 #include "PushModelUtils.h"
 #include "Net/Core/PushModel/PushModel.h"
 #include "Containers/SparseArray.h"
+#include "UObject/ObjectKey.h"
 
 namespace UE4PushModelPrivate
 {
@@ -25,17 +26,17 @@ namespace UE4PushModelPrivate
 		/**
 		 * Creates a new FPushModelPerObjectState.
 		 *
-		 * @param InObjectId			The ID that we'll use to refer to this object. Should be unique across all Push Model Objects (but doesn't necessarily have to be UObject::GetUniqueId()).
+		 * @param InObjectKey			An ObjectKey that can uniqely identify the object we're representing.
 		 * @param InNumberOfProperties	The total number of replicated properties this object has.
 		 */
-		FPushModelPerObjectState(const FNetPushObjectId InObjectId, const uint16 InNumberOfProperties)
-			: ObjectId(InObjectId)
+		FPushModelPerObjectState(const FObjectKey InObjectKey, const uint16 InNumberOfProperties)
+			: ObjectKey(InObjectKey)
 			, DirtiedThisFrame(true, InNumberOfProperties)
 		{
 		}
 
 		FPushModelPerObjectState(FPushModelPerObjectState&& Other)
-			: ObjectId(Other.ObjectId)
+			: ObjectKey(Other.ObjectKey)
 			, DirtiedThisFrame(MoveTemp(Other.DirtiedThisFrame))
 			, PerNetDriverStates(MoveTemp(Other.PerNetDriverStates))
 		{
@@ -117,7 +118,7 @@ namespace UE4PushModelPrivate
 			return DirtiedThisFrame.Num();
 		}
 
-		const FNetPushObjectId GetObjectId() const
+		const FObjectKey& GetObjectKey() const
 		{
 			return ObjectId;
 		}
@@ -125,7 +126,7 @@ namespace UE4PushModelPrivate
 	private:
 	
 		//! A unique ID for the object.
-		const FNetPushObjectId ObjectId;
+		const FObjectKey ObjectKey;
 
 		//! Bitfield tracking which properties we've dirtied since the last time
 		//! our state was pushed to NetDrivers.
diff --git a/Engine/Source/Runtime/Net/Core/Public/Net/Core/PushModel/PushModel.h b/Engine/Source/Runtime/Net/Core/Public/Net/Core/PushModel/PushModel.h
--- a/Engine/Source/Runtime/Net/Core/Public/Net/Core/PushModel/PushModel.h
+++ b/Engine/Source/Runtime/Net/Core/Public/Net/Core/PushModel/PushModel.h
@@ -7,6 +7,7 @@
 #if WITH_PUSH_MODEL
 
 #include "CoreMinimal.h"
+#include "UObject/ObjectKey.h"
 
 /**
  * Push Model Support for networking.
@@ -341,7 +342,7 @@ namespace UE4PushModelPrivate
 	NETCORE_API void MarkPropertyDirty(const FNetPushObjectId ObjectId, const int32 StartRepIndex, const int32 EndRepIndex);
 
 
-	NETCORE_API const FPushModelPerNetDriverHandle AddPushModelObject(const FNetPushObjectId ObjectId, const uint16 NumberOfReplicatedProperties);
+	NETCORE_API const FPushModelPerNetDriverHandle AddPushModelObject(const FObjectKey ObjectId, const uint16 NumberOfReplicatedProperties);
 	NETCORE_API void RemovePushModelObject(const FPushModelPerNetDriverHandle Handle);
 
 	NETCORE_API class FPushModelPerNetDriverState* GetPerNetDriverState(const FPushModelPerNetDriverHandle Handle);
