Fix UE4 references.

diff --git a/Engine/Source/Runtime/Net/Core/Private/Net/Core/Misc/PacketAudit.cpp b/Engine/Source/Runtime/Net/Core/Private/Net/Core/Misc/PacketAudit.cpp
--- a/Engine/Source/Runtime/Net/Core/Private/Net/Core/Misc/PacketAudit.cpp
+++ b/Engine/Source/Runtime/Net/Core/Private/Net/Core/Misc/PacketAudit.cpp
@@ -17,7 +17,7 @@
  */
 
 /** The general name used/adapted for mutexes and shared memory */
-#define AUDIT_MUTEX_NAME	TEXT("UE4PacketAudit")
+#define AUDIT_MUTEX_NAME	TEXT("UEPacketAudit")
 
 /** Hardcoded amount of shared memory reserved for both the send/receive maps */
 #define AUDIT_MAPPING_SIZE	(1024 * 1024 * 32)
@@ -253,7 +253,7 @@ FPacketAudit::FPacketAudit(FPlatformProcess::FSemaphore* InGameMutex)
 	// @todo #JohnB: There is probably a way to fix the shared memory platform code, so this works without admin. Low priority though.
 	if (SendSharedMemory == nullptr || ReceiveSharedMemory == nullptr)
 	{
-		LowLevelFatalError(TEXT("You need to run UE4 as administrator, for packet auditing to work."));
+		LowLevelFatalError(TEXT("You need to run UnrealEngine as administrator, for packet auditing to work."));
 	}
 
 	FSharedMemoryWriter WipeSendSharedMemory((uint8*)SendSharedMemory->GetAddress());
