NetTrace events that were using attachments have been converted over to use string or array type fields.

#rnx

diff --git a/Engine/Source/Developer/TraceServices/Private/Analyzers/NetTraceAnalyzer.cpp b/Engine/Source/Developer/TraceServices/Private/Analyzers/NetTraceAnalyzer.cpp
--- a/Engine/Source/Developer/TraceServices/Private/Analyzers/NetTraceAnalyzer.cpp
+++ b/Engine/Source/Developer/TraceServices/Private/Analyzers/NetTraceAnalyzer.cpp
@@ -115,7 +115,8 @@ bool FNetTraceAnalyzer::OnEvent(uint16 RouteId, EStyle Style, const FOnEventCont
 			}
 			else
 			{
-				TracedNameIdToNetProfilerNameIdMap.Add(TraceNameId, NetProfilerProvider.AddNetProfilerName(UTF8_TO_TCHAR(EventData.GetAttachment())));
+				FString Name = FTraceAnalyzerUtils::LegacyAttachmentString<UTF8CHAR>("Name", Context);
+				TracedNameIdToNetProfilerNameIdMap.Add(TraceNameId, NetProfilerProvider.AddNetProfilerName(*Name));
 			}
 		}
 		break;
@@ -189,8 +190,9 @@ void FNetTraceAnalyzer::HandlePacketContentEvent(const FOnEventContext& Context,
 	TArray<FBunchInfo>& BunchInfos = (ConnectionState->BunchInfos)[ConnectionMode];
 
 	// Decode batched events
-	uint64 BufferSize = EventData.GetAttachmentSize();
-	const uint8* BufferPtr = EventData.GetAttachment();
+	TArrayView<const uint8> DataView = FTraceAnalyzerUtils::LegacyAttachmentArray("Data", Context);
+	uint64 BufferSize = DataView.Num();
+	const uint8* BufferPtr = DataView.GetData();
 	const uint8* BufferEnd = BufferPtr + BufferSize;
 	uint64 LastOffset = 0;
 
diff --git a/Engine/Source/Runtime/Net/Core/Private/Net/Core/Trace/Reporters/NetTraceReporter.cpp b/Engine/Source/Runtime/Net/Core/Private/Net/Core/Trace/Reporters/NetTraceReporter.cpp
--- a/Engine/Source/Runtime/Net/Core/Private/Net/Core/Trace/Reporters/NetTraceReporter.cpp
+++ b/Engine/Source/Runtime/Net/Core/Private/Net/Core/Trace/Reporters/NetTraceReporter.cpp
@@ -25,6 +25,7 @@ UE_TRACE_EVENT_END()
 // Trace a name, the utf encoded name is attached as a attachment
 UE_TRACE_EVENT_BEGIN(NetTrace, NameEvent)
 	UE_TRACE_EVENT_FIELD(uint16, NameId)
+	UE_TRACE_EVENT_FIELD(UE::Trace::AnsiString, Name)
 UE_TRACE_EVENT_END()
 
 UE_TRACE_EVENT_BEGIN(NetTrace, ObjectCreatedEvent)	
@@ -62,6 +63,7 @@ UE_TRACE_EVENT_BEGIN(NetTrace, PacketContentEvent)
 	UE_TRACE_EVENT_FIELD(uint16, ConnectionId)
 	UE_TRACE_EVENT_FIELD(uint8, GameInstanceId)
 	UE_TRACE_EVENT_FIELD(uint8, PacketType)
+	UE_TRACE_EVENT_FIELD(uint8[], Data)
 UE_TRACE_EVENT_END()
 
 //$TODO: Drop the timestamp when we can get them for free on the analysis side
@@ -98,9 +100,9 @@ void FNetTraceReporter::ReportInstanceDestroyed(uint32 GameInstanceId)
 
 void FNetTraceReporter::ReportAnsiName(FNetDebugNameId NameId, uint32 NameSize, const char* Name)
 {
-	UE_TRACE_LOG(NetTrace, NameEvent, NetChannel, NameSize)
+	UE_TRACE_LOG(NetTrace, NameEvent, NetChannel)
 		<< NameEvent.NameId(NameId)
-		<< NameEvent.Attachment(Name, NameSize);
+		<< NameEvent.Name(Name, NameSize);
 }
 
 void FNetTraceReporter::ReportPacketDropped(const FNetTracePacketInfo& PacketInfo)
@@ -138,11 +140,11 @@ void FNetTraceReporter::ReportPacketContent(FNetTracePacketContentEvent* Events,
 
 	auto FlushPacketContentBuffer = [](const FNetTracePacketInfo& InPacketInfo, const uint8* InBuffer, uint32 Count)
 	{
-		UE_TRACE_LOG(NetTrace, PacketContentEvent, NetChannel, Count)
+		UE_TRACE_LOG(NetTrace, PacketContentEvent, NetChannel)
 			<< PacketContentEvent.ConnectionId(InPacketInfo.ConnectionId)
 			<< PacketContentEvent.GameInstanceId(InPacketInfo.GameInstanceId)
 			<< PacketContentEvent.PacketType((uint8)InPacketInfo.PacketType)
-			<< PacketContentEvent.Attachment(InBuffer, Count);
+			<< PacketContentEvent.Data(InBuffer, Count);
 	};
 
 	for (const FNetTracePacketContentEvent& CurrentEvent : MakeArrayView(Events, EventCount))
