Type truncation warning fixes in a few networking classes

diff --git a/Engine/Source/Runtime/CoreUObject/Private/UObject/Class.cpp b/Engine/Source/Runtime/CoreUObject/Private/UObject/Class.cpp
--- a/Engine/Source/Runtime/CoreUObject/Private/UObject/Class.cpp
+++ b/Engine/Source/Runtime/CoreUObject/Private/UObject/Class.cpp
@@ -4252,10 +4252,11 @@ void UClass::SetUpUhtReplicationData()
 		{
 			if (It->PropertyFlags & CPF_Net)
 			{
-				It->RepIndex = ClassReps.Num();
+				It->RepIndex = (uint16)ClassReps.Num();
 				new (ClassReps) FRepRecord(*It, 0);
 			}
 		}
+		check(ClassReps.Num() <= 65535);
 
 		ClassFlags |= CLASS_ReplicationDataIsSetUp;
 		ClassReps.Shrink();
@@ -4334,12 +4335,13 @@ void UClass::SetUpRuntimeReplicationData()
 		ClassReps.Reserve(ClassReps.Num() + NetProperties.Num());
 		for (int32 i = 0; i < NetProperties.Num(); i++)
 		{
-			NetProperties[i]->RepIndex = ClassReps.Num();
+			NetProperties[i]->RepIndex = (uint16)ClassReps.Num();
 			for (int32 j = 0; j < NetProperties[i]->ArrayDim; j++)
 			{
 				new(ClassReps)FRepRecord(NetProperties[i], j);
 			}
 		}
+		check(ClassReps.Num() <= 65535);
 
 		if (bIsNativeClass && GValidateReplicatedProperties)
 		{
diff --git a/Engine/Source/Runtime/CoreUObject/Public/UObject/CoreNet.h b/Engine/Source/Runtime/CoreUObject/Public/UObject/CoreNet.h
--- a/Engine/Source/Runtime/CoreUObject/Public/UObject/CoreNet.h
+++ b/Engine/Source/Runtime/CoreUObject/Public/UObject/CoreNet.h
@@ -269,7 +269,7 @@ public:
 	}
 
 	FLifetimeProperty(int32 InRepIndex)
-		: RepIndex(InRepIndex)
+		: RepIndex((uint16)InRepIndex)
 		, Condition(COND_None)
 		, RepNotifyCondition(REPNOTIFY_OnChanged)
 		, bIsPushBased(false)
@@ -278,7 +278,7 @@ public:
 	}
 
 	FLifetimeProperty(int32 InRepIndex, ELifetimeCondition InCondition, ELifetimeRepNotifyCondition InRepNotifyCondition=REPNOTIFY_OnChanged, bool bInIsPushBased=false)
-		: RepIndex(InRepIndex)
+		: RepIndex((uint16)InRepIndex)
 		, Condition(InCondition)
 		, RepNotifyCondition(InRepNotifyCondition)
 		, bIsPushBased(bInIsPushBased)
diff --git a/Engine/Source/Runtime/Engine/Public/Net/NetAnalyticsTypes.h b/Engine/Source/Runtime/Engine/Public/Net/NetAnalyticsTypes.h
--- a/Engine/Source/Runtime/Engine/Public/Net/NetAnalyticsTypes.h
+++ b/Engine/Source/Runtime/Engine/Public/Net/NetAnalyticsTypes.h
@@ -343,7 +343,7 @@ public:
 		return NumberOfTrackedFrames;
 	}
 
-	const float GetBurstyPacketLossPerSecond(double DurationSec) const
+	const double GetBurstyPacketLossPerSecond(double DurationSec) const
 	{
 		return DurationSec > 0.0 ? GetNumberOfFramesWithBurstsOfPacketLoss() / DurationSec : 0.0;
 	}
diff --git a/Engine/Source/Runtime/Engine/Public/Net/TrafficControl.h b/Engine/Source/Runtime/Engine/Public/Net/TrafficControl.h
--- a/Engine/Source/Runtime/Engine/Public/Net/TrafficControl.h
+++ b/Engine/Source/Runtime/Engine/Public/Net/TrafficControl.h
@@ -44,7 +44,7 @@ public:
 
 	double GetLatestRTT() const { return LatestRTT; }
 	uint32 GetBytesInFlight() const { return BytesInFlight; }
-	uint32 GetPacketsInFlight() const { return OutPacketRecords.Count(); }
+	uint32 GetPacketsInFlight() const { return (uint32)OutPacketRecords.Count(); }
 
 private:
 	FNetworkCongestionControl* TrafficControlModule;
diff --git a/Engine/Source/Runtime/Net/Core/Public/Net/Core/Misc/DDoSDetection.h b/Engine/Source/Runtime/Net/Core/Public/Net/Core/Misc/DDoSDetection.h
--- a/Engine/Source/Runtime/Net/Core/Public/Net/Core/Misc/DDoSDetection.h
+++ b/Engine/Source/Runtime/Net/Core/Public/Net/Core/Misc/DDoSDetection.h
@@ -155,9 +155,9 @@ struct NETCORE_API FDDoSStateConfig : public FDDoSState
 	{
 		// Exclude escalation triggers from this
 		//Target.EscalateTimeQuotaMSPerFrame		= (EscalateTimeQuotaMSPerFrame == -1 ? -1 : EscalateTimeQuotaMSPerFrame * FrameAdjustment);
-		Target.PacketLimitPerFrame				= (PacketLimitPerFrame == -1 ? -1 : PacketLimitPerFrame * FrameAdjustment);
-		Target.PacketTimeLimitMSPerFrame		= (PacketTimeLimitMSPerFrame == -1 ? -1 : PacketTimeLimitMSPerFrame * FrameAdjustment);
-		Target.NetConnPacketTimeLimitMSPerFrame	= (NetConnPacketTimeLimitMSPerFrame == -1 ? -1 : NetConnPacketTimeLimitMSPerFrame * FrameAdjustment);
+		Target.PacketLimitPerFrame				= (PacketLimitPerFrame == -1 ? -1 : (int32)(PacketLimitPerFrame * FrameAdjustment));
+		Target.PacketTimeLimitMSPerFrame		= (PacketTimeLimitMSPerFrame == -1 ? -1 : (int32)(PacketTimeLimitMSPerFrame * FrameAdjustment));
+		Target.NetConnPacketTimeLimitMSPerFrame	= (NetConnPacketTimeLimitMSPerFrame == -1 ? -1 : (int32)(NetConnPacketTimeLimitMSPerFrame * FrameAdjustment));
 	}
 };
 
