Push Model;

Adding a compare/assign and mark dirty macro to the push model.

-- Code --
PushModel:
COMPARE_ASSIGN_AND_MARK_PROPERTY_DIRTY - Created so anything can use it.

diff --git a/Engine/Source/Runtime/Net/Core/Public/Net/Core/PushModel/PushModel.h b/Engine/Source/Runtime/Net/Core/Public/Net/Core/PushModel/PushModel.h
--- a/Engine/Source/Runtime/Net/Core/Public/Net/Core/PushModel/PushModel.h
+++ b/Engine/Source/Runtime/Net/Core/Public/Net/Core/PushModel/PushModel.h
@@ -424,6 +424,12 @@ namespace UEPushModelPrivate
 // Marks an entire static array property dirty, given the Class Name, Property Name, and Object. This will fail to compile if the Property or Class aren't valid.
 #define MARK_PROPERTY_DIRTY_FROM_NAME_STATIC_ARRAY(ClassName, PropertyName, Object) CONDITIONAL_ON_OBJECT_NET_ID(Object, UEPushModelPrivate::MarkPropertyDirty(PrivatePushId, GET_PROPERTY_REP_INDEX_STATIC_ARRAY_START(ClassName, PropertyName), GET_PROPERTY_REP_INDEX_STATIC_ARRAY_END(ClassName, PropertyName))
 
+/**
+ * This is used to reduce lines of code (mostly in setters) by doing the comparison check before changing the value and marking dirty only if we have changed the value. 
+ * This is best if used for things like ints, floats, pointers, but not something like a struct because of the memcmp cost. 
+ */
+#define COMPARE_ASSIGN_AND_MARK_PROPERTY_DIRTY(ClassName, PropertyName, NewValue, Object) \
+	if (NewValue != PropertyName) { PropertyName = NewValue; MARK_PROPERTY_DIRTY_FROM_NAME(ClassName, PropertyName, Object); }
 
 #else // WITH_PUSH_MODEL
 
