Push Model;

Adding a compare/assign and mark dirty macro to the push model.

-- Code --
PushModel:
COMPARE_ASSIGN_AND_MARK_PROPERTY_DIRTY - Defining in the #else so it'll compile on the client.

diff --git a/Engine/Source/Runtime/Net/Core/Public/Net/Core/PushModel/PushModel.h b/Engine/Source/Runtime/Net/Core/Public/Net/Core/PushModel/PushModel.h
--- a/Engine/Source/Runtime/Net/Core/Public/Net/Core/PushModel/PushModel.h
+++ b/Engine/Source/Runtime/Net/Core/Public/Net/Core/PushModel/PushModel.h
@@ -448,4 +448,7 @@ namespace UEPushModelPrivate
 #define IS_PUSH_MODEL_ENABLED() false
 #define PUSH_MAKE_BP_PROPERTIES_PUSH_MODEL() false
 
+#define COMPARE_ASSIGN_AND_MARK_PROPERTY_DIRTY(ClassName, PropertyName, NewValue, Object) \
+	if (NewValue != PropertyName) { PropertyName = NewValue; }
+
 #endif
\ No newline at end of file
