Type truncation fixes - Net

diff --git a/Engine/Source/Runtime/Net/Core/Classes/Net/Serialization/FastArraySerializer.h b/Engine/Source/Runtime/Net/Core/Classes/Net/Serialization/FastArraySerializer.h
--- a/Engine/Source/Runtime/Net/Core/Classes/Net/Serialization/FastArraySerializer.h
+++ b/Engine/Source/Runtime/Net/Core/Classes/Net/Serialization/FastArraySerializer.h
@@ -1407,7 +1407,8 @@ bool FFastArraySerializer::FastArrayDeltaSerialize(TArray<Type> &Items, FNetDelt
 					GuidReferences.Buffer.Empty();
 
 					// Remember the number of bits in the buffer
-					GuidReferences.NumBufferBits = Reader.GetPosBits() - Mark.GetPos();
+					check(Reader.GetPosBits() - Mark.GetPos() <= TNumericLimits<int32>::Max());
+					GuidReferences.NumBufferBits = int32(Reader.GetPosBits() - Mark.GetPos());
 					
 					// Copy the buffer itself
 					Mark.Copy( Reader, GuidReferences.Buffer );
diff --git a/Engine/Source/Runtime/Net/Core/Public/Net/Core/Misc/ResizableCircularQueue.h b/Engine/Source/Runtime/Net/Core/Public/Net/Core/Misc/ResizableCircularQueue.h
--- a/Engine/Source/Runtime/Net/Core/Public/Net/Core/Misc/ResizableCircularQueue.h
+++ b/Engine/Source/Runtime/Net/Core/Public/Net/Core/Misc/ResizableCircularQueue.h
@@ -260,7 +260,8 @@ void TResizableCircularQueue<T, AllocatorT>::PopNoCheck(SIZE_T Count)
 		}
 	}
 
-	Tail += Count;
+	check(SIZE_T(Tail) + Count <= TNumericLimits<IndexT>::Max());
+	Tail += (IndexT)Count;
 }
 
 template<typename T, typename AllocatorT>
