Manually integrating fast array fix from 18643809

diff --git a/Engine/Source/Runtime/Net/Core/Classes/Net/Serialization/FastArraySerializer.h b/Engine/Source/Runtime/Net/Core/Classes/Net/Serialization/FastArraySerializer.h
--- a/Engine/Source/Runtime/Net/Core/Classes/Net/Serialization/FastArraySerializer.h
+++ b/Engine/Source/Runtime/Net/Core/Classes/Net/Serialization/FastArraySerializer.h
@@ -1601,8 +1601,8 @@ bool FFastArraySerializer::FastArrayDeltaSerialize_DeltaSerializeStructs(TArray<
 
 		// We set ChangelistHistory and LastAckedHistory to 0, as that's the default RepLayout expects.
 		// Setting them to anything else may cause issues with initial sends.
-		int32 OldChangelistHistory = 0;
-		int32 OldLastAckedHistory = 0;
+		uint32 OldChangelistHistory = 0;
+		uint32 OldLastAckedHistory = 0;
 
 		// See if the array changed at all. If the ArrayReplicationKey matches we can skip checking individual items
 		if (Parms.OldState)
