Branching NetBitArray from Dev-Iris

diff --git a/Engine/Source/Runtime/Net/Core/Public/Net/Core/NetBitArray.h b/Engine/Source/Runtime/Net/Core/Public/Net/Core/NetBitArray.h
new file mode 100644
--- /dev/null
+++ b/Engine/Source/Runtime/Net/Core/Public/Net/Core/NetBitArray.h
@@ -0,0 +1,1103 @@
+// Copyright Epic Games, Inc. All Rights Reserved.
+
+#pragma once
+
+#include "HAL/Platform.h"
+#include "HAL/PlatformMath.h"
+#include "Containers/Array.h"
+
+/* NetBitArray validation support. */
+#ifndef UE_NETBITARRAY_VALIDATE
+#if (UE_BUILD_SHIPPING || UE_BUILD_TEST)
+#	define UE_NETBITARRAY_VALIDATE 0
+#else
+#	define UE_NETBITARRAY_VALIDATE 1
+#endif
+#endif
+
+#if UE_NETBITARRAY_VALIDATE
+#	define UE_NETBITARRAY_CHECK(x) check(x)
+#else
+#	define UE_NETBITARRAY_CHECK(...)
+#endif
+
+#define UE_NETBITARRAY_VALIDATE_COMPATIBLE(x) UE_NETBITARRAY_CHECK(this != &(x) && BitCount == (x).BitCount)
+#define UE_NETBITARRAYVIEW_VALIDATE_COMPATIBLE(x) UE_NETBITARRAY_CHECK((Storage != (x).Storage) && (BitCount == (x).BitCount))
+
+namespace UE::Net
+{
+
+class FNetBitArrayBase
+{
+public:
+	typedef uint32 StorageWordType;
+	static constexpr uint32 WordBitCount = sizeof(StorageWordType)*8U;
+
+	// Some WordOps that can be passed to Combine() and ForAllSetBits()
+	static constexpr StorageWordType AndOp(StorageWordType A, StorageWordType B) { return A & B; }
+	static constexpr StorageWordType AndNotOp(StorageWordType A, StorageWordType B) { return A & ~B; }
+	static constexpr StorageWordType OrOp(StorageWordType A, StorageWordType B) { return A | B; }
+	static constexpr StorageWordType XorOp(StorageWordType A, StorageWordType B) { return A ^ B; }
+};
+
+/**
+ * Simple bit array with internal storage.
+ * Has very little error checking as it is used in performance critical code.
+ */
+class FNetBitArray : public FNetBitArrayBase
+{
+public:
+	/**  Creates an empty NetBitArray. */
+	FNetBitArray();
+
+	/**  Creates an array with size BitCount bits. */
+	explicit FNetBitArray(uint32 BitCount);
+
+	/** Return true if equal including BitCount and padding bits */
+	bool operator==(const FNetBitArray& Other) const;
+
+	/**
+	 * Sets a new size and resets the array's contents.
+	 * @param BitCount - The number of bits in the array after the operation.
+	 */
+	void Init(uint32 BitCount);
+
+	/**
+	 * Release all memory and set capacity to 0 bits
+	 */
+	void Empty();
+
+	/**
+	 * Sets the number of bits to BitCount. If the array grows the new bits will be cleared.
+	 * Existing bits will remain intact.
+	 *
+	 * @param BitCount - The number of bits in the array after the operation.
+	 */
+	void SetNumBits(uint32 BitCount);
+
+	/**
+	 * Grows the size with the specified number of bits. New bits will be cleared.
+	 *
+	 * @param BitCount - The number of bits to add to the array.
+	 * @see SetNumBits
+	 */
+	void AddBits(uint32 BitCount);
+
+	/** Returns the number of bits */
+	uint32 GetNumBits() const;
+
+	/** Returns the number of words in our storage */
+	uint32 GetNumWords() const;
+
+	/** Returns a pointer to the internal storage. */
+	const StorageWordType* GetData() const { return Storage.GetData(); }
+
+	/** Returns a pointer to the internal storage. */
+	StorageWordType* GetData() { return Storage.GetData(); }
+
+	/** Clear all bits in the array. */
+	void Reset();
+
+	/** Sets all bits in the array, but clears padding bits. */
+	void SetAllBits();
+
+	/** Returns true if any bit is set in the bitset. Note: Padding bits in storage are expected to be zero. */
+	bool IsAnyBitSet() const;
+
+	/** Returns true if any bit is set in the specified range. The range will be clamped if needed.  Note: Padding bits in storage are expected to be zero. */
+	bool IsAnyBitSet(uint32 StartIndex, uint32 Count) const;
+
+	/** Returns true if no bit is set in the bitset. Note: Padding bits in storage are expected to be zero. */
+	bool IsNoBitSet() const;
+
+	/** Set the bit with the specified index */
+	void SetBit(uint32 Index);
+
+	/** Set the bit with the specified index to bValue */
+	void SetBitValue(uint32 Index, bool bValue);
+
+	/** Set the bit with the specified index to bValue */
+	void SetBits(uint32 StartIndex, uint32 Count);
+
+	/** Clear the bits in the specified range */
+	void ClearBits(uint32 StartIndex, uint32 Count);
+
+	/** Clear the bit with the specified index */
+	void ClearBit(uint32 Index);
+
+	/** Get the bit with the specified index */
+	bool GetBit(uint32 Index) const;
+
+	/** Find first zero bit. Returns ~0U if no zero bit was found */
+	uint32 FindFirstZero() const;
+
+	/** Find first set bit. Returns ~0U if no zero bit was found */
+	uint32 FindFirstOne() const;
+
+	/** Find first zero bit starting from StartIndex. Returns ~0U if no zero bit was found. */
+	uint32 FindFirstZero(uint32 StartIndex) const;
+
+	/** Find first set bit starting from StartIndex. Returns ~0U if no set bit was found. */
+	uint32 FindFirstOne(uint32 StartIndex) const;
+
+	/** Find last zero bit. Returns ~0U if no zero bit was found. */
+	uint32 FindLastZero() const;
+
+	/** Find last set bit. Returns ~0U if no set bit was found. */
+	uint32 FindLastOne() const;
+
+	/**
+	 * Retrieves set bits in the provided range and returns how many indices were written to OutIndices.
+	 * If OutIndices is filled to its capacity the search for set bits will end. OutIndices may be modified
+	 * beyond the returned count but within the capacity.
+	 * 
+	 * @param StartIndex Which bit index to start searching for set bits.
+	 * @param Count How many bits to check. The count will be clamped to number of bits in the array so you can pass ~0U to check all bits.
+	 * @param OutIndices Where to store the indices of the set bits. Indices with lower numbers will be stored first.
+	 * @param OutIndicesCapacity How many indices can be stored in OutIndices.
+	 * @returns Number of indices OutIndices was populated with.
+	 */
+	uint32 GetSetBitIndices(uint32 StartIndex, uint32 Count, uint32* OutIndices, uint32 OutIndicesCapacity) const;
+
+	/** Counts the number of set bits in this array in the provided range. */
+	uint32 CountSetBits(uint32 StartIndex = 0, uint32 Count = ~0U) const;
+
+	/** Copy bits from other bit array. The other array must have the same bit count. */
+	void Copy(const FNetBitArray& Other);
+
+	/**
+	 * Combine this array with another array using a word operation functor. Bit counts need to match.
+	 *
+	 * @param Other The other array to combine with.
+	 * @param WordOpFunctor The operation to perform for each word to combine. Signature is StorageWordType(StorageWordType ThisWord, StorageWordType OtherWord).
+	 * @see FNetBitArrayHelper::AndOp, FNetBitArrayHelper::AndNotOp, FNetBitArrayHelper::OrOp, FNetBitArrayHelper::XorOp
+	 */
+	template<typename WordOpFunctor> void Combine(const FNetBitArray& Other, WordOpFunctor&& WordOp);
+
+	/** Iterate over all set bits and invoke functor with signature void(uint32 BitIndex). */
+	template<typename T>
+	void ForAllSetBits(T&& Functor) const;
+
+	/** Iterate over all words and execute Functor, with signature void(uint32 BitIndex), for all set bits after executing WordOpFunctor */
+	template<typename T, typename V>
+	static void ForAllSetBits(const FNetBitArray& A, const FNetBitArray& B, T&& WordOpFunctor, V&& Functor);
+
+	/** Compare two BitArrays bit by bit and invoke FunctorA(BitIndex) for bits set only in array A and invoke FunctorB(BitIndex) for bits set only in array B. */
+	template<typename T, typename V>
+	static void ForAllExclusiveBits(const FNetBitArray& A, const FNetBitArray& B, T&& FunctorA, V&& FunctorB);
+
+private:
+	StorageWordType GetLastWordMask() const { return (~StorageWordType(0) >> (uint32(-int32(BitCount)) & (WordBitCount - 1))); }
+	void ClearPaddingBits();
+
+	TArray<StorageWordType> Storage;
+	uint32 BitCount;
+};
+
+/**
+ * Simple bit array view.
+ * Assumes external storage and has very little error checking as it is used in performance critical code.
+ * Note: Trailing padding bits in storage will be zeroed unless the view is explicitly constructed to allow it.
+ *       Certain operations assume the padding bits are cleared.
+ */
+class FNetBitArrayView : public FNetBitArrayBase
+{
+public:
+	enum ENoResetNoValidateType { NoResetNoValidate };
+	enum EResetOnInitType { ResetOnInit };
+
+	/** Constructor will produce a valid but empty bitarray. */
+	inline FNetBitArrayView();
+
+	/** Construct NetBitArray from external storage. Storage must be large enough to fit given the BitCount. No clearing of storage and no validation that padding bits are cleared. */
+	inline FNetBitArrayView(StorageWordType* StorageIn, uint32 BitCountIn, const ENoResetNoValidateType);
+	/** Construct NetBitArray from external storage. Storage must be large enough to fit given the BitCount. Clears the entire storage. */
+	inline FNetBitArrayView(StorageWordType* StorageIn, uint32 BitCountIn, const EResetOnInitType);
+	/** Construct NetBitArray from external storage. Storage must be large enough to fit given the BitCount. Validates that padding bits are cleared. */
+	inline FNetBitArrayView(StorageWordType* StorageIn, uint32 BitCountIn);
+
+	/** Return true if equal including BitCount and padding bits */
+	bool operator==(const FNetBitArrayView& Other) const;
+
+	/** Returns true if any bit is set in the bitset Note: Padding bits in storage are expected to be zero. */
+	inline bool IsAnyBitSet() const;
+
+	/** Returns true if any bit is set in the specified range. The range will be clamped if needed. Note: Padding bits in storage are expected to be zero.*/
+	inline bool IsAnyBitSet(uint32 StartIndex, uint32 Count) const;
+
+	/** Returns true if no bit is set in the bitset. Note: Padding bits in storage are expected to be zero. */
+	inline bool IsNoBitSet() const;
+
+	/** Reset the storage of the BitArray including any padding bits */
+	inline void Reset();
+
+	/** All padding bits will be set to zero */
+	inline void ClearPaddingBits();
+
+	/** Set all bits, padding bits will be zeroed*/
+	inline void SetAllBits();
+
+	/** Set the bit with the specified index */
+	inline void SetBit(uint32 Index);
+
+	/** Set the bit with the specified index to bValue */
+	inline void SetBitValue(uint32 Index, bool bValue);
+
+	/** Set the bits in the specified range to true */
+	inline void SetBits(uint32 StartIndex, uint32 Count);
+
+	/** Clear the bits in the specified range */
+	inline void ClearBits(uint32 StartIndex, uint32 Count);
+
+	/** Clear the bit with the specified index */
+	inline void ClearBit(uint32 Index) const;
+
+	/** Get the bit with the specified index */
+	inline bool GetBit(uint32 Index) const;
+
+	/** Find first zero bit. Returns ~0U if no zero bit was found */
+	inline uint32 FindFirstZero() const;
+
+	/** Find first set bit. Returns ~0U if no zero bit was found */
+	inline uint32 FindFirstOne() const;
+
+	/** Find first zero bit starting from StartIndex. Returns ~0U if no zero bit was found. */
+	inline uint32 FindFirstZero(uint32 StartIndex) const;
+
+	/** Find first set bit starting from StartIndex. Returns ~0U if no set bit was found. */
+	inline uint32 FindFirstOne(uint32 StartIndex) const;
+
+	/** Find last zero bit. Returns ~0U if no zero bit was found. */
+	uint32 FindLastZero() const;
+
+	/** Find last set bit. Returns ~0U if no set bit was found. */
+	uint32 FindLastOne() const;
+
+	/** Returns the number of bits */
+	inline uint32 GetNumBits() const;
+
+	/** Returns the number of words in our storage */
+	inline uint32 GetNumWords() const;
+
+	/** Returns a pointer to the internal storage. */
+	const StorageWordType* GetData() const { return Storage; }
+
+	/** Returns a pointer to the internal storage. */
+	StorageWordType* GetData() { return Storage; }
+
+	/**
+	 * Retrieves set bits in the provided range and returns how many indices were written to OutIndices.
+	 * If OutIndices is filled to its capacity the search for set bits will end. OutIndices may be modified
+	 * beyond the returned count but within the capacity.
+	 * 
+	 * @param StartIndex Which bit index to start searching for set bits.
+	 * @param Count How many bits to check. The count will be clamped to number of bits in the array so you can pass ~0U to check all bits.
+	 * @param OutIndices Where to store the indices of the set bits. Indices with lower numbers will be stored first.
+	 * @param OutIndicesCapacity How many indices can be stored in OutIndices.
+	 * @returns Number of indices OutIndices was populated with.
+	 */
+	uint32 GetSetBitIndices(uint32 StartIndex, uint32 Count, uint32* OutIndices, uint32 OutIndicesCapacity) const;
+
+	/** Counts the number of set bits in this array in the provided range. */
+	uint32 CountSetBits(uint32 StartIndex = 0, uint32 Count = ~0U) const;
+
+	/** Copy bits from other bit array. The other array must have the same bit count. */
+	inline void Copy(const FNetBitArrayView& Other);
+
+	/**
+	 * Combine this array with another array using a word operation functor. Bit counts need to match.
+	 *
+	 * @param Other The other array to combine with.
+	 * @param WordOpFunctor The operation to perform for each word to combine. Signature is StorageWordType(StorageWordType ThisWord, StorageWordType OtherWord).
+	 * @see FNetBitArrayHelper::AndOp, FNetBitArrayHelper::AndNotOp, FNetBitArrayHelper::OrOp, FNetBitArrayHelper::XorOp
+	 */
+	template<typename Functor> void Combine(const FNetBitArrayView& Other, Functor&& WordOp);
+
+	/** Iterate over all set bits and invoke functor with signature void(uint32 BitIndex). */
+	template <typename T>
+	void ForAllSetBits(T&& Functor) const;
+
+	/** Iterate over all words and execute Functor, with signature void(uint32 BitIndex), for all set bits after executing WordOpFunctor */
+	template <typename T, typename V>
+	static void ForAllSetBits(const FNetBitArrayView& A, const FNetBitArrayView& B, T&& WordOpFunctor, V&& Functor);
+
+	/** Compare two BitArrays bit by bit and invoke FunctorA(BitIndex) for bits set only in array A and invoke FunctorB(BitIndex) for bits set only in array B. */
+	template <typename T, typename V>
+	static void ForAllExclusiveBits(const FNetBitArrayView& A, const FNetBitArrayView& B, T&& FunctorA, V&& FunctorB);
+
+	static constexpr inline uint32 CalculateRequiredWordCount(uint32 BitCount);
+
+private:
+	inline StorageWordType GetLastWordMask() const { return (~StorageWordType(0) >> (uint32(-int32(BitCount)) & (WordBitCount - 1))); }
+
+	StorageWordType* Storage;
+	uint32 WordCount;
+	uint32 BitCount;
+};
+
+
+// Implementation helper
+class FNetBitArrayHelper final
+{
+private:
+	friend FNetBitArray;
+	friend FNetBitArrayView;
+
+	typedef FNetBitArray::StorageWordType StorageWordType;
+	static constexpr uint32 WordBitCount = sizeof(StorageWordType) * 8u;
+
+	static StorageWordType GetLastWordMask(uint32 BitCount)
+	{
+		return (~StorageWordType(0) >> (uint32(-int32(BitCount)) & (WordBitCount - 1)));
+	}
+
+	static bool IsAnyBitSet(const StorageWordType* Storage, const uint32 WordCount)
+	{
+		for (uint32 WordIt = 0; WordIt < WordCount; ++WordIt)
+		{
+			if (Storage[WordIt])
+			{
+				return true;
+			}
+		}
+
+		return false;
+	}
+
+	static bool IsAnyBitSet(const StorageWordType* Storage, const uint32 BitCount, const uint32 StartIndex, const uint32 Count)
+	{
+		// Range validation
+		if ((Count == 0) | (StartIndex >= BitCount))
+		{
+			return false;
+		}
+
+		const uint32 EndIndex = static_cast<uint32>(FPlatformMath::Min(uint64(StartIndex) + Count - 1U, uint64(BitCount) - 1U));
+		const uint32 WordStartIt = StartIndex / WordBitCount;
+		StorageWordType StartWordMask = (~StorageWordType(0) << (StartIndex & (WordBitCount - 1)));
+		const uint32 WordEndIt = EndIndex / WordBitCount;
+		const StorageWordType EndWordMask = (~StorageWordType(0) >> (~EndIndex & (WordBitCount - 1)));
+
+		for (uint32 WordIt = WordStartIt; WordIt <= WordEndIt; ++WordIt)
+		{
+			// Compute mask based on start and end word status.
+			const bool bIsEndWord = (WordIt == WordEndIt);
+			const StorageWordType WordMask = StartWordMask & (bIsEndWord ? EndWordMask : ~StorageWordType(0));
+			// Only the first word is the start word so we can set all bits now.
+			StartWordMask = ~StorageWordType(0);
+
+			if (Storage[WordIt] & WordMask)
+			{
+				return true;
+			}
+		}
+
+		return false;
+	}
+
+	static void SetBits(StorageWordType* Storage, const uint32 BitCount, const uint32 StartIndex, const uint32 Count)
+	{
+		UE_NETBITARRAY_CHECK((Count > 0) & (StartIndex < BitCount));
+
+		const uint32 EndIndex = static_cast<uint32>(FPlatformMath::Min(uint64(StartIndex) + Count - 1U, uint64(BitCount) - 1U));
+		const uint32 WordStartIt = StartIndex / WordBitCount;
+		StorageWordType StartWordMask = (~StorageWordType(0) << (StartIndex & (WordBitCount - 1)));
+		const uint32 WordEndIt = EndIndex / WordBitCount;
+		const StorageWordType EndWordMask = (~StorageWordType(0) >> (~EndIndex & (WordBitCount - 1)));
+
+		for (uint32 WordIt = WordStartIt; WordIt <= WordEndIt; ++WordIt)
+		{
+			// Compute mask based on start and end word status.
+			const bool bIsEndWord = (WordIt == WordEndIt);
+			const StorageWordType WordMask = StartWordMask & (bIsEndWord ? EndWordMask : ~StorageWordType(0));
+			// Only the first word is the start word so we can set all bits now.
+			StartWordMask = ~StorageWordType(0);
+
+			Storage[WordIt] |= ~StorageWordType(0) & WordMask;
+		}
+	}
+
+	static void ClearBits(StorageWordType* Storage, const uint32 BitCount, const uint32 StartIndex, const uint32 Count)
+	{
+		UE_NETBITARRAY_CHECK((Count > 0) & (StartIndex < BitCount));
+
+		const uint32 EndIndex = static_cast<uint32>(FPlatformMath::Min(uint64(StartIndex) + Count - 1U, uint64(BitCount) - 1U));
+		const uint32 WordStartIt = StartIndex / WordBitCount;
+		StorageWordType StartWordMask = (~StorageWordType(0) << (StartIndex & (WordBitCount - 1)));
+		const uint32 WordEndIt = EndIndex / WordBitCount;
+		const StorageWordType EndWordMask = (~StorageWordType(0) >> (~EndIndex & (WordBitCount - 1)));
+
+		for (uint32 WordIt = WordStartIt; WordIt <= WordEndIt; ++WordIt)
+		{
+			// Compute mask based on start and end word status.
+			const bool bIsEndWord = (WordIt == WordEndIt);
+			const StorageWordType WordMask = StartWordMask & (bIsEndWord ? EndWordMask : ~StorageWordType(0));
+			// Only the first word is the start word so we can set all bits now.
+			StartWordMask = ~StorageWordType(0);
+
+			Storage[WordIt] = Storage[WordIt] & ~WordMask;
+		}
+	}
+
+	static void SetBit(StorageWordType* Storage, const uint32 BitCount, const uint32 Index)
+	{
+		UE_NETBITARRAY_CHECK(Index < BitCount);
+
+		const SIZE_T WordIndex = Index/WordBitCount;
+		const StorageWordType WordMask = (StorageWordType(1) << (Index & (WordBitCount - 1)));
+
+		Storage[WordIndex] |= WordMask;
+	}
+
+	static void SetBitValue(StorageWordType* Storage, const uint32 BitCount, const uint32 Index, const bool bValue)
+	{
+		UE_NETBITARRAY_CHECK(Index < BitCount);
+
+		const SIZE_T WordIndex = Index/WordBitCount;
+		const StorageWordType WordMask = (StorageWordType(1) << (Index & (WordBitCount - 1)));
+		const StorageWordType ValueMask = bValue ? WordMask : StorageWordType(0);
+
+		Storage[WordIndex] = (Storage[WordIndex] & ~WordMask) | ValueMask;
+	}
+
+	static void ClearBit(StorageWordType* Storage, const uint32 BitCount, const uint32 Index)
+	{
+		UE_NETBITARRAY_CHECK(Index < BitCount);
+
+		const SIZE_T WordIndex = Index/WordBitCount;
+		const StorageWordType WordMask = (StorageWordType(1) << (Index & (WordBitCount - 1)));
+
+		Storage[WordIndex] &= ~WordMask;
+	}
+
+	static bool GetBit(const StorageWordType* Storage, const uint32 BitCount, const uint32 Index)
+	{
+		UE_NETBITARRAY_CHECK(Index < BitCount);
+
+		const SIZE_T WordIndex = Index/WordBitCount;
+		const StorageWordType WordMask = (StorageWordType(1) << (Index & (WordBitCount - 1)));
+
+		return (Storage[WordIndex] & WordMask) != 0u;
+	}
+
+	static uint32 FindFirstZero(const StorageWordType* Storage, const uint32 WordCount, const uint32 BitCount)
+	{
+		for (uint32 WordIt = 0, CurrentBitIndex = 0; WordIt < WordCount; ++WordIt, CurrentBitIndex += WordBitCount)
+		{
+			const StorageWordType Word = Storage[WordIt];
+			if (Word != StorageWordType(~StorageWordType(0)))
+			{
+				const uint32 Index = CurrentBitIndex + FPlatformMath::CountTrailingZeros(~Word);
+				// Need to make sure the index is not out of bounds
+				return (Index < BitCount ? Index : ~0U);
+			}
+		}
+
+		return ~0U;
+	}
+
+	static uint32 FindFirstOne(const StorageWordType* Storage, const uint32 WordCount, const uint32 BitCount)
+	{
+		for (uint32 WordIt = 0, CurrentBitIndex = 0; WordIt < WordCount; ++WordIt, CurrentBitIndex += WordBitCount)
+		{
+			const StorageWordType Word = Storage[WordIt];
+			if (Word != 0)
+			{
+				const uint32 Index = CurrentBitIndex + FPlatformMath::CountTrailingZeros(Word);
+				// Need to make sure the index is not out of bounds
+				return (Index < BitCount ? Index : ~0U);
+			}
+		}
+
+		return ~0U;
+	}
+
+	static uint32 FindFirstZero(const StorageWordType* Storage, const uint32 WordCount, const uint32 BitCount, const uint32 StartIndex)
+	{
+		// WordMask is used to hide zero bits in the first word we check to prevent returning an index before StartIndex.
+		uint32 CurrentBitIndex = StartIndex & ~(WordBitCount - 1U);
+		uint32 WordMask = ~(~0U << (StartIndex & (WordBitCount - 1U)));
+		for (uint32 WordIt = StartIndex/WordBitCount; WordIt < WordCount; ++WordIt, CurrentBitIndex += WordBitCount)
+		{
+			const StorageWordType Word = Storage[WordIt] | WordMask;
+			WordMask = 0U;
+			if (Word != StorageWordType(~StorageWordType(0)))
+			{
+				const uint32 Index = CurrentBitIndex + FPlatformMath::CountTrailingZeros(~Word);
+				// Need to make sure the index is not out of bounds
+				return (Index < BitCount ? Index : ~0U);
+			}
+		}
+
+		return ~0U;
+	}
+
+	static uint32 FindFirstOne(const StorageWordType* Storage, const uint32 WordCount, const uint32 BitCount, const uint32 StartIndex)
+	{
+		uint32 CurrentBitIndex = StartIndex & ~(WordBitCount - 1U);
+		// WordMask is needed only in the first word we check to prevent returning an index before StartIndex.
+		uint32 WordMask = ~0U << (StartIndex & (WordBitCount - 1U));
+		for (uint32 WordIt = StartIndex/WordBitCount; WordIt < WordCount; ++WordIt, CurrentBitIndex += WordBitCount)
+		{
+			const StorageWordType Word = Storage[WordIt] & WordMask;
+			// From here on we accept all words as they are.
+			WordMask = ~0U;
+			if (Word != 0)
+			{
+				const uint32 Index = CurrentBitIndex + FPlatformMath::CountTrailingZeros(Word);
+				// Need to make sure the index is not out of bounds
+				return (Index < BitCount ? Index : ~0U);
+			}
+		}
+
+		return ~0U;
+	}
+
+	static uint32 FindLastZero(const StorageWordType* Storage, const uint32 WordCount, const uint32 BitCount)
+	{
+		StorageWordType WordMask = GetLastWordMask(BitCount);
+
+		for (uint32 WordIt = WordCount; WordIt > 0; )
+		{
+			--WordIt;
+			const StorageWordType Word = Storage[WordIt];
+			if (Word != WordMask)
+			{
+				const uint32 BitOffset = WordIt*WordBitCount;
+				const uint32 Index = BitOffset + WordBitCount - 1U - FPlatformMath::CountLeadingZeros(~Word);
+				return Index;
+			}
+
+			WordMask = StorageWordType(~StorageWordType(0));
+		}
+
+		return ~0U;
+	}
+
+	static uint32 FindLastOne(const StorageWordType* Storage, const uint32 WordCount, const uint32 BitCount)
+	{
+		for (uint32 WordIt = WordCount; WordIt > 0; )
+		{
+			--WordIt;
+			const StorageWordType Word = Storage[WordIt];
+			if (Word != 0U)
+			{
+				const uint32 BitOffset = WordIt*WordBitCount;
+				const uint32 Index = BitOffset + WordBitCount - 1U - FPlatformMath::CountLeadingZeros(Word);
+				return Index;
+			}
+		}
+
+		return ~0U;
+	}
+
+	NETCORE_API static uint32 GetSetBitIndices(const StorageWordType* Storage, const uint32 BitCount, const uint32 StartIndex, const uint32 Count, uint32* const OutIndices, const uint32 OutIndicesCapacity);
+	
+	NETCORE_API static uint32 CountSetBits(const StorageWordType* Storage, const uint32 BitCount, const uint32 StartIndex, const uint32 Count);
+
+	template<typename Functor> static void Combine(StorageWordType* Storage, const StorageWordType* OtherStorage, const uint32 WordCount, Functor&& WordOp)
+	{
+		for (uint32 WordIt = 0; WordIt < WordCount; ++WordIt)
+		{
+			Storage[WordIt] = WordOp(Storage[WordIt], OtherStorage[WordIt]);
+		}
+	}
+
+	static bool IsEqual(const StorageWordType* Storage, const StorageWordType* OtherStorage, const uint32 WordCount)
+	{
+		StorageWordType Result = 0U;
+		for (uint32 WordIt = 0; WordIt < WordCount; ++WordIt)
+		{
+			Result |= Storage[WordIt] ^ OtherStorage[WordIt];
+		}
+		return Result == 0U;
+	}
+
+	template<typename T> static void ForAllSetBits(const StorageWordType* Storage, const uint32 WordCount, const uint32 BitCount, T&& Functor)
+	{
+		const StorageWordType LastWordMask = ~StorageWordType(0) >> (uint32(-int32(BitCount)) & (WordBitCount - 1));
+		const uint32 LastWordIt = WordCount - 1;
+
+		for (uint32 WordIt = 0, CurrentBitIndex = 0; WordIt < WordCount; ++WordIt, CurrentBitIndex += WordBitCount)
+		{
+			StorageWordType CurrentWord = Storage[WordIt] & ((WordIt == LastWordIt) ? LastWordMask : ~0U);
+
+			// Test All bits in the CurrentWord and invoke functor if they are set
+			uint32 LocalBitIndex = CurrentBitIndex;
+			while (CurrentWord)
+			{
+				if (CurrentWord & 0x1)
+				{
+					Functor(LocalBitIndex);
+				}
+				CurrentWord >>= 1;
+				++LocalBitIndex;
+			}
+		}
+	}
+
+	template<typename T, typename V> static void ForAllSetBits(const StorageWordType* StorageA, const StorageWordType* StorageB, const uint32 WordCount, const uint32 BitCount, T&& WordOpFunctor, V&& Functor)
+	{
+		const StorageWordType LastWordMask = ~StorageWordType(0) >> (uint32(-int32(BitCount)) & (WordBitCount - 1));
+		const uint32 LastWordIt = WordCount - 1;
+
+		for (uint32 WordIt = 0, CurrentBitIndex = 0; WordIt < WordCount; ++WordIt, CurrentBitIndex += WordBitCount)
+		{
+			const StorageWordType CurrentWordA = StorageA[WordIt];
+			const StorageWordType CurrentWordB = StorageB[WordIt];
+
+			StorageWordType CurrentWord = WordOpFunctor(CurrentWordA, CurrentWordB) & ((WordIt == LastWordIt) ? LastWordMask : ~0U);
+
+			// Test All bits in the CurrentWord and invoke functor if they are set
+			uint32 LocalBitIndex = CurrentBitIndex;
+			while (CurrentWord)
+			{
+				if (CurrentWord & 0x1)
+				{
+					Functor(LocalBitIndex);
+				}
+				CurrentWord >>= 1;
+				++LocalBitIndex;
+			}
+		}
+	}
+
+	template<typename T, typename V> static void ForAllExclusiveBits(const StorageWordType* StorageA, const StorageWordType* StorageB, const uint32 WordCount, const uint32 BitCount, T&& FunctorA, V&& FunctorB)
+	{
+		const StorageWordType LastWordMask = ~StorageWordType(0) >> (uint32(-int32(BitCount)) & (WordBitCount - 1));
+		const uint32 LastWordIt = WordCount - 1;
+
+		for (uint32 WordIt = 0, CurrentBitIndex = 0; WordIt < WordCount; ++WordIt, CurrentBitIndex += WordBitCount)
+		{
+			StorageWordType CurrentWordA = StorageA[WordIt];
+			const StorageWordType CurrentWordB = StorageB[WordIt];
+			StorageWordType CurrentWordXOR = (CurrentWordA ^ CurrentWordB) & ((WordIt == LastWordIt) ? LastWordMask : ~0U);
+
+			// if CurrentWord contains any set bits invoke the correct functor
+			uint32 LocalBitIndex = CurrentBitIndex;
+			while (CurrentWordXOR)
+			{
+				if (CurrentWordXOR & 0x1)
+				{
+					if (CurrentWordA & 0x1)
+					{
+						FunctorA(LocalBitIndex);
+					}
+					else
+					{
+						FunctorB(LocalBitIndex);
+					}
+				}
+
+				CurrentWordXOR >>= 1;
+				CurrentWordA >>= 1;
+
+				++LocalBitIndex;
+			}
+		}
+	}
+
+};
+
+
+// NetBitArray Implementation
+inline FNetBitArray::FNetBitArray()
+: BitCount(0)
+{
+}
+
+inline bool FNetBitArray::operator==(const FNetBitArray& Other) const
+{
+	if (BitCount != Other.BitCount)
+	{
+		return false;
+	}
+
+	return FNetBitArrayHelper::IsEqual(this->GetData(), Other.GetData(), GetNumWords());
+}
+
+inline FNetBitArray::FNetBitArray(uint32 InBitCount)
+: BitCount(InBitCount)
+{
+	Init(InBitCount);
+}
+
+inline void FNetBitArray::Init(uint32 InBitCount)
+{
+	BitCount = InBitCount;
+
+	const uint32 WordCount = (InBitCount + WordBitCount - 1U)/WordBitCount;
+	Storage.Reset(WordCount);
+	Storage.AddZeroed(WordCount);
+}
+
+inline void FNetBitArray::Empty()
+{
+	BitCount = 0U;
+	Storage.Empty();
+};
+
+inline void FNetBitArray::SetNumBits(uint32 InBitCount)
+{
+	BitCount = InBitCount;
+
+	const uint32 WordCount = (InBitCount + WordBitCount - 1U)/WordBitCount;
+	constexpr bool bAllowShrinking = false;
+	Storage.SetNumZeroed(WordCount, bAllowShrinking);
+	ClearPaddingBits();
+}
+
+inline void FNetBitArray::AddBits(uint32 AdditionalBitCount)
+{
+	SetNumBits(BitCount + AdditionalBitCount);
+}
+
+inline uint32 FNetBitArray::GetNumBits() const
+{
+	return BitCount;
+}
+
+inline uint32 FNetBitArray::GetNumWords() const
+{
+	return static_cast<uint32>(Storage.Num());
+}
+
+inline void FNetBitArray::Reset()
+{
+	FPlatformMemory::Memset(Storage.GetData(), 0, Storage.Num()*sizeof(StorageWordType));
+}
+
+inline void FNetBitArray::SetAllBits()
+{
+	FPlatformMemory::Memset(Storage.GetData(), 0xff, Storage.Num()*sizeof(StorageWordType));
+	ClearPaddingBits();
+}
+
+inline bool FNetBitArray::IsAnyBitSet() const
+{
+	return FNetBitArrayHelper::IsAnyBitSet(Storage.GetData(), Storage.Num());
+}
+
+inline bool FNetBitArray::IsAnyBitSet(uint32 StartIndex, uint32 Count) const
+{
+	return FNetBitArrayHelper::IsAnyBitSet(Storage.GetData(), BitCount, StartIndex, Count);
+}
+
+inline bool FNetBitArray::IsNoBitSet() const
+{
+	return !IsAnyBitSet();
+}
+
+inline void FNetBitArray::SetBit(uint32 Index)
+{
+	FNetBitArrayHelper::SetBit(Storage.GetData(), BitCount, Index);
+}
+
+inline void FNetBitArray::SetBitValue(uint32 Index, bool bValue)
+{
+	FNetBitArrayHelper::SetBitValue(Storage.GetData(), BitCount, Index, bValue);
+}
+
+inline void FNetBitArray::SetBits(uint32 StartIndex, uint32 Count)
+{
+	FNetBitArrayHelper::SetBits(Storage.GetData(), BitCount, StartIndex, Count);
+}
+
+inline void FNetBitArray::ClearBits(uint32 StartIndex, uint32 Count)
+{
+	FNetBitArrayHelper::ClearBits(Storage.GetData(), BitCount, StartIndex, Count);
+}
+
+inline void FNetBitArray::ClearBit(uint32 Index)
+{
+	FNetBitArrayHelper::ClearBit(Storage.GetData(), BitCount, Index);
+}
+
+inline bool FNetBitArray::GetBit(uint32 Index) const
+{
+	return FNetBitArrayHelper::GetBit(Storage.GetData(), BitCount, Index);
+}
+
+inline uint32 FNetBitArray::FindFirstZero() const
+{
+	return FNetBitArrayHelper::FindFirstZero(Storage.GetData(), Storage.Num(), BitCount);
+}
+
+inline uint32 FNetBitArray::FindFirstOne() const
+{
+	return FNetBitArrayHelper::FindFirstOne(Storage.GetData(), Storage.Num(), BitCount);
+}
+
+inline uint32 FNetBitArray::FindFirstZero(uint32 StartIndex) const
+{
+	return FNetBitArrayHelper::FindFirstZero(Storage.GetData(), Storage.Num(), BitCount, StartIndex);
+}
+
+inline uint32 FNetBitArray::FindFirstOne(uint32 StartIndex) const
+{
+	return FNetBitArrayHelper::FindFirstOne(Storage.GetData(), Storage.Num(), BitCount, StartIndex);
+}
+
+inline uint32 FNetBitArray::FindLastZero() const
+{
+	return FNetBitArrayHelper::FindLastZero(Storage.GetData(), Storage.Num(), BitCount);
+}
+
+inline uint32 FNetBitArray::FindLastOne() const
+{
+	return FNetBitArrayHelper::FindLastOne(Storage.GetData(), Storage.Num(), BitCount);
+}
+
+inline uint32 FNetBitArray::GetSetBitIndices(uint32 StartIndex, uint32 Count, uint32* OutIndices, uint32 OutIndicesCapacity) const
+{
+	return FNetBitArrayHelper::GetSetBitIndices(Storage.GetData(), BitCount, StartIndex, Count, OutIndices, OutIndicesCapacity);
+}
+
+inline uint32 FNetBitArray::CountSetBits(uint32 StartIndex, uint32 Count) const
+{
+	return FNetBitArrayHelper::CountSetBits(Storage.GetData(), BitCount, StartIndex, Count);
+}
+
+inline void FNetBitArray::Copy(const FNetBitArray& Other)
+{
+	UE_NETBITARRAY_VALIDATE_COMPATIBLE(Other);
+	// Intentionally doing a memory copy instead of using assignment operator as we know the bit counts are the same and don't need extra checks.
+	FPlatformMemory::Memcpy(GetData(), Other.GetData(), GetNumWords()*sizeof(StorageWordType));
+}
+
+template<typename Functor> inline void FNetBitArray::Combine(const FNetBitArray& Other, Functor&& WordOp)
+{
+	UE_NETBITARRAY_VALIDATE_COMPATIBLE(Other);
+	FNetBitArrayHelper::Combine(Storage.GetData(), Other.Storage.GetData(), Storage.Num(), WordOp);
+}
+
+template<typename T> inline void FNetBitArray::ForAllSetBits(T&& Functor) const
+{
+	FNetBitArrayHelper::ForAllSetBits(Storage.GetData(), Storage.Num(), BitCount, Functor);
+}
+
+template <typename T, typename V> inline void FNetBitArray::ForAllSetBits(const FNetBitArray& A, const FNetBitArray& B, T&& WordOpFunctor, V&& Functor)
+{
+	UE_NETBITARRAY_CHECK(A.GetNumBits() == B.GetNumBits());
+	FNetBitArrayHelper::ForAllSetBits(A.Storage.GetData(), B.Storage.GetData(), A.Storage.Num(), A.BitCount, WordOpFunctor, Functor);
+}
+
+template <typename T, typename V> inline void FNetBitArray::ForAllExclusiveBits(const FNetBitArray& A, const FNetBitArray& B, T&& FunctorA, V&& FunctorB)
+{
+	UE_NETBITARRAY_CHECK(A.GetNumBits() == B.GetNumBits());
+	FNetBitArrayHelper::ForAllExclusiveBits(A.Storage.GetData(), B.Storage.GetData(), A.Storage.Num(), A.BitCount, FunctorA, FunctorB);
+}
+
+inline void FNetBitArray::ClearPaddingBits()
+{
+	if (BitCount > 0)
+	{
+		StorageWordType& LastWord = Storage.GetData()[Storage.Num() - 1];
+		LastWord &= GetLastWordMask();
+	}
+}
+
+// NetBitArrayView implementation
+constexpr uint32 FNetBitArrayView::CalculateRequiredWordCount(uint32 BitCount)
+{ 
+	return (BitCount + WordBitCount - 1) / WordBitCount; 
+}
+
+FNetBitArrayView::FNetBitArrayView()
+: Storage(&BitCount)
+, WordCount(0)
+, BitCount(0)
+{
+}
+
+FNetBitArrayView::FNetBitArrayView(StorageWordType* StorageIn, uint32 BitCountIn, const ENoResetNoValidateType)
+: Storage(StorageIn)
+, WordCount(CalculateRequiredWordCount(BitCountIn))
+, BitCount(BitCountIn)
+{
+	if (BitCountIn == 0)
+	{
+		Storage = const_cast<uint32*>(&BitCount);
+	}
+}
+
+FNetBitArrayView::FNetBitArrayView(StorageWordType* StorageIn, uint32 BitCountIn, const EResetOnInitType)
+: FNetBitArrayView(StorageIn, BitCountIn, NoResetNoValidate)
+{
+	Reset();
+}
+
+FNetBitArrayView::FNetBitArrayView(StorageWordType* StorageIn, uint32 BitCountIn)
+: FNetBitArrayView(StorageIn, BitCountIn, NoResetNoValidate)
+{
+	UE_NETBITARRAY_CHECK(Storage != nullptr);
+	UE_NETBITARRAY_CHECK((BitCount == 0) || ((Storage[WordCount - 1] & ~GetLastWordMask()) == 0u));
+}
+
+inline bool FNetBitArrayView::operator==(const FNetBitArrayView& Other) const
+{
+	if (BitCount != Other.BitCount)
+	{
+		return false;
+	}
+
+	return FNetBitArrayHelper::IsEqual(this->GetData(), Other.GetData(), GetNumWords());
+}
+
+bool FNetBitArrayView::IsAnyBitSet() const
+{
+	return FNetBitArrayHelper::IsAnyBitSet(Storage, WordCount);
+}
+
+bool FNetBitArrayView::IsAnyBitSet(uint32 StartIndex, uint32 Count) const
+{
+	return FNetBitArrayHelper::IsAnyBitSet(Storage, BitCount, StartIndex, Count);
+}
+
+bool FNetBitArrayView::IsNoBitSet() const
+{
+	return !IsAnyBitSet();
+}
+
+void FNetBitArrayView::Reset()
+{
+	FPlatformMemory::Memset(&Storage[0], 0, WordCount * sizeof(StorageWordType));
+}
+
+void FNetBitArrayView::ClearPaddingBits()
+{
+	if (BitCount > 0)
+	{
+		Storage[WordCount - 1] = Storage[WordCount - 1] & GetLastWordMask();
+	}
+}
+
+void FNetBitArrayView::SetAllBits()
+{
+	FPlatformMemory::Memset(&Storage[0], 0xff, WordCount * sizeof(StorageWordType));
+	ClearPaddingBits();
+}
+
+/** Copy bits from other bit array */
+void FNetBitArrayView::Copy(const FNetBitArrayView& Other)
+{
+	UE_NETBITARRAYVIEW_VALIDATE_COMPATIBLE(Other);
+	FPlatformMemory::Memcpy(&Storage[0], &Other.Storage[0], WordCount * sizeof(StorageWordType));
+}
+
+void FNetBitArrayView::SetBit(uint32 Index)
+{
+	FNetBitArrayHelper::SetBit(Storage, BitCount, Index);
+}
+
+void FNetBitArrayView::SetBitValue(uint32 Index, bool bValue)
+{
+	FNetBitArrayHelper::SetBitValue(Storage, BitCount, Index, bValue);
+}
+
+void FNetBitArrayView::SetBits(uint32 StartIndex, uint32 Count)
+{
+	FNetBitArrayHelper::SetBits(Storage, BitCount, StartIndex, Count);
+}
+
+void FNetBitArrayView::ClearBits(uint32 StartIndex, uint32 Count)
+{
+	FNetBitArrayHelper::ClearBits(Storage, BitCount, StartIndex, Count);
+}
+
+void FNetBitArrayView::ClearBit(uint32 Index) const
+{
+	FNetBitArrayHelper::ClearBit(Storage, BitCount, Index);
+}
+
+bool FNetBitArrayView::GetBit(uint32 Index) const
+{
+	return FNetBitArrayHelper::GetBit(Storage, BitCount, Index);
+}
+
+uint32 FNetBitArrayView::FindFirstZero() const
+{
+	return FNetBitArrayHelper::FindFirstZero(Storage, WordCount, BitCount);
+}
+
+uint32 FNetBitArrayView::FindFirstOne() const
+{
+	return FNetBitArrayHelper::FindFirstOne(Storage, WordCount, BitCount);
+}
+
+uint32 FNetBitArrayView::FindFirstZero(uint32 StartIndex) const
+{
+	return FNetBitArrayHelper::FindFirstZero(Storage, WordCount, BitCount, StartIndex);
+}
+
+uint32 FNetBitArrayView::FindFirstOne(uint32 StartIndex) const
+{
+	return FNetBitArrayHelper::FindFirstOne(Storage, WordCount, BitCount, StartIndex);
+}
+
+inline uint32 FNetBitArrayView::FindLastZero() const
+{
+	return FNetBitArrayHelper::FindLastZero(Storage, WordCount, BitCount);
+}
+
+inline uint32 FNetBitArrayView::FindLastOne() const
+{
+	return FNetBitArrayHelper::FindLastOne(Storage, WordCount, BitCount);
+}
+
+uint32 FNetBitArrayView::GetNumWords() const
+{
+	return WordCount;
+}
+
+uint32 FNetBitArrayView::GetNumBits() const
+{
+	return BitCount;
+}
+
+inline uint32 FNetBitArrayView::GetSetBitIndices(uint32 StartIndex, uint32 Count, uint32* OutIndices, uint32 OutIndicesCapacity) const
+{
+	return FNetBitArrayHelper::GetSetBitIndices(Storage, BitCount, StartIndex, Count, OutIndices, OutIndicesCapacity);
+}
+
+inline uint32 FNetBitArrayView::CountSetBits(uint32 StartIndex, uint32 Count) const
+{
+	return FNetBitArrayHelper::CountSetBits(Storage, BitCount, StartIndex, Count);
+}
+
+template<typename Functor> inline void FNetBitArrayView::Combine(const FNetBitArrayView& Other, Functor&& WordOp)
+{
+	UE_NETBITARRAYVIEW_VALIDATE_COMPATIBLE(Other);
+	FNetBitArrayHelper::Combine(Storage, Other.Storage, WordCount, WordOp);
+}
+
+template<typename T> inline void FNetBitArrayView::ForAllSetBits(T&& Functor) const
+{
+	FNetBitArrayHelper::ForAllSetBits(Storage, WordCount, BitCount, Functor);
+}
+
+template <typename T, typename V>
+void FNetBitArrayView::ForAllSetBits(const FNetBitArrayView& A, const FNetBitArrayView& B, T&& WordOpFunctor, V&& Functor)
+{
+	UE_NETBITARRAY_CHECK(A.GetNumBits() == B.GetNumBits());
+	FNetBitArrayHelper::ForAllSetBits(A.Storage, B.Storage, A.WordCount, A.BitCount, WordOpFunctor, Functor);
+}
+
+template <typename T, typename V>
+void FNetBitArrayView::ForAllExclusiveBits(const FNetBitArrayView& A, const FNetBitArrayView& B, T&& FunctorA, V&& FunctorB)
+{
+	UE_NETBITARRAY_CHECK(A.GetNumBits() == B.GetNumBits());
+	FNetBitArrayHelper::ForAllExclusiveBits(A.Storage, B.Storage, A.WordCount, A.BitCount, FunctorA, FunctorB);
+}
+
+inline FNetBitArrayView MakeNetBitArrayView(const FNetBitArrayView::StorageWordType* Storage, uint32 BitCount)
+{
+	return FNetBitArrayView(const_cast<FNetBitArrayView::StorageWordType*>(Storage), BitCount);
+}
+
+inline FNetBitArrayView MakeNetBitArrayView(const FNetBitArray& BitArray)
+{
+	return FNetBitArrayView(const_cast<FNetBitArrayView::StorageWordType*>(BitArray.GetData()), BitArray.GetNumBits());
+}
+
+} // end namespace UE::Net
