Fixed FSampleMinMaxAvg not detecting special cases for Min/Max values.

diff --git a/Engine/Source/Runtime/Net/Core/Public/Net/Core/Analytics/NetStatsUtils.h b/Engine/Source/Runtime/Net/Core/Public/Net/Core/Analytics/NetStatsUtils.h
--- a/Engine/Source/Runtime/Net/Core/Public/Net/Core/Analytics/NetStatsUtils.h
+++ b/Engine/Source/Runtime/Net/Core/Public/Net/Core/Analytics/NetStatsUtils.h
@@ -384,7 +384,7 @@ public:
 	 */
 	double GetMin() const
 	{
-		return MinValue;
+		return ((MinValue == std::numeric_limits<double>::max()) ? 0.0 : MinValue);
 	}
 
 	/**
@@ -394,7 +394,7 @@ public:
 	 */
 	double GetMax() const
 	{
-		return MaxValue;
+		return ((MaxValue == std::numeric_limits<double>::min()) ? 0.0 : MaxValue);
 	}
 
 	/**
