Renamed Net/Core/Trace/Config.h file to NetTraceConfig to prevent include mismatch with other Trace/Config.h file

diff --git a/Engine/Source/Runtime/Engine/Classes/Engine/NetConnection.h b/Engine/Source/Runtime/Engine/Classes/Engine/NetConnection.h
--- a/Engine/Source/Runtime/Engine/Classes/Engine/NetConnection.h
+++ b/Engine/Source/Runtime/Engine/Classes/Engine/NetConnection.h
@@ -22,7 +22,7 @@
 #include "ProfilingDebugging/Histogram.h"
 #include "Containers/ArrayView.h"
 #include "Containers/CircularBuffer.h"
-#include "Net/Core/Trace/Config.h"
+#include "Net/Core/Trace/NetTraceConfig.h"
 #include "ReplicationDriver.h"
 #include "Analytics/EngineNetAnalytics.h"
 #include "Net/Common/Packets/PacketTraits.h"
diff --git a/Engine/Source/Runtime/Engine/Public/Net/DataBunch.h b/Engine/Source/Runtime/Engine/Public/Net/DataBunch.h
--- a/Engine/Source/Runtime/Engine/Public/Net/DataBunch.h
+++ b/Engine/Source/Runtime/Engine/Public/Net/DataBunch.h
@@ -9,7 +9,7 @@
 #include "Misc/NetworkGuid.h"
 #include "UObject/CoreNet.h"
 #include "EngineLogs.h"
-#include "Net/Core/Trace/Config.h"
+#include "Net/Core/Trace/NetTraceConfig.h"
 
 class UChannel;
 class UNetConnection;
diff --git a/Engine/Source/Runtime/Net/Core/Public/Net/Core/Trace/Config.h b/Engine/Source/Runtime/Net/Core/Public/Net/Core/Trace/Config.h
deleted file mode 100644
diff --git a/Engine/Source/Runtime/Net/Core/Public/Net/Core/Trace/NetTraceConfig.h b/Engine/Source/Runtime/Net/Core/Public/Net/Core/Trace/NetTraceConfig.h
new file mode 100644
--- /dev/null
+++ b/Engine/Source/Runtime/Net/Core/Public/Net/Core/Trace/NetTraceConfig.h
@@ -0,0 +1,35 @@
+// Copyright Epic Games, Inc. All Rights Reserved.
+
+#pragma once
+
+#include "Trace/Config.h"
+
+// Forward declarations
+class FNetTraceCollector;
+
+#if !defined(UE_NET_TRACE_ENABLED)
+#	if !(UE_BUILD_SHIPPING || UE_BUILD_TEST) && UE_TRACE_ENABLED
+#		define UE_NET_TRACE_ENABLED 1
+#	else
+#		define UE_NET_TRACE_ENABLED 0
+#	endif
+#endif
+
+namespace ENetTraceVerbosity
+{
+	enum Type : uint32
+	{
+		None = 0,
+		Trace,
+		Verbose,
+		VeryVerbose,
+	};
+}
+
+#if UE_NET_TRACE_ENABLED
+#	ifndef UE_NET_TRACE_COMPILETIME_VERBOSITY
+#		define UE_NET_TRACE_COMPILETIME_VERBOSITY ENetTraceVerbosity::Verbose
+#	endif
+#else
+#	define UE_NET_TRACE_COMPILETIME_VERBOSITY ENetTraceVerbosity::None
+#endif
\ No newline at end of file
diff --git a/Engine/Source/Runtime/Net/Core/Public/Net/Core/Trace/Private/NetTraceInternal.h b/Engine/Source/Runtime/Net/Core/Public/Net/Core/Trace/Private/NetTraceInternal.h
--- a/Engine/Source/Runtime/Net/Core/Public/Net/Core/Trace/Private/NetTraceInternal.h
+++ b/Engine/Source/Runtime/Net/Core/Public/Net/Core/Trace/Private/NetTraceInternal.h
@@ -3,7 +3,7 @@
 #pragma once
 
 #include "CoreTypes.h"
-#include "Net/Core/Trace/Config.h"
+#include "Net/Core/Trace/NetTraceConfig.h"
 
 #if UE_NET_TRACE_ENABLED
 
