Iris - Add InvalidIndex to represent failure in various Find methods.

#[review]-19826837

diff --git a/Engine/Source/Runtime/Net/Core/Public/Net/Core/NetBitArray.h b/Engine/Source/Runtime/Net/Core/Public/Net/Core/NetBitArray.h
--- a/Engine/Source/Runtime/Net/Core/Public/Net/Core/NetBitArray.h
+++ b/Engine/Source/Runtime/Net/Core/Public/Net/Core/NetBitArray.h
@@ -33,6 +33,9 @@ public:
 	typedef uint32 StorageWordType;
 	static constexpr uint32 WordBitCount = sizeof(StorageWordType)*8U;
 
+	/** Returned from various methods to indicate failure. */
+	static constexpr uint32 InvalidIndex = ~0U;
+
 	// Some WordOps that can be passed to Combine() and ForAllSetBits()
 	static constexpr StorageWordType AndOp(StorageWordType A, StorageWordType B) { return A & B; }
 	static constexpr StorageWordType AndNotOp(StorageWordType A, StorageWordType B) { return A & ~B; }
@@ -128,22 +131,22 @@ public:
 	/** Get the bit with the specified index */
 	bool GetBit(uint32 Index) const;
 
-	/** Find first zero bit. Returns ~0U if no zero bit was found */
+	/** Find first zero bit. Returns InvalidIndex if no zero bit was found */
 	uint32 FindFirstZero() const;
 
-	/** Find first set bit. Returns ~0U if no zero bit was found */
+	/** Find first set bit. Returns InvalidIndex if no zero bit was found */
 	uint32 FindFirstOne() const;
 
-	/** Find first zero bit starting from StartIndex. Returns ~0U if no zero bit was found. */
+	/** Find first zero bit starting from StartIndex. Returns InvalidIndex if no zero bit was found. */
 	uint32 FindFirstZero(uint32 StartIndex) const;
 
-	/** Find first set bit starting from StartIndex. Returns ~0U if no set bit was found. */
+	/** Find first set bit starting from StartIndex. Returns InvalidIndex if no set bit was found. */
 	uint32 FindFirstOne(uint32 StartIndex) const;
 
-	/** Find last zero bit. Returns ~0U if no zero bit was found. */
+	/** Find last zero bit. Returns InvalidIndex if no zero bit was found. */
 	uint32 FindLastZero() const;
 
-	/** Find last set bit. Returns ~0U if no set bit was found. */
+	/** Find last set bit. Returns InvalidIndex if no set bit was found. */
 	uint32 FindLastOne() const;
 
 	/**
@@ -255,22 +258,22 @@ public:
 	/** Get the bit with the specified index */
 	inline bool GetBit(uint32 Index) const;
 
-	/** Find first zero bit. Returns ~0U if no zero bit was found */
+	/** Find first zero bit. Returns InvalidIndex if no zero bit was found */
 	inline uint32 FindFirstZero() const;
 
-	/** Find first set bit. Returns ~0U if no zero bit was found */
+	/** Find first set bit. Returns InvalidIndex if no zero bit was found */
 	inline uint32 FindFirstOne() const;
 
-	/** Find first zero bit starting from StartIndex. Returns ~0U if no zero bit was found. */
+	/** Find first zero bit starting from StartIndex. Returns InvalidIndex if no zero bit was found. */
 	inline uint32 FindFirstZero(uint32 StartIndex) const;
 
-	/** Find first set bit starting from StartIndex. Returns ~0U if no set bit was found. */
+	/** Find first set bit starting from StartIndex. Returns InvalidIndex if no set bit was found. */
 	inline uint32 FindFirstOne(uint32 StartIndex) const;
 
-	/** Find last zero bit. Returns ~0U if no zero bit was found. */
+	/** Find last zero bit. Returns InvalidIndex if no zero bit was found. */
 	uint32 FindLastZero() const;
 
-	/** Find last set bit. Returns ~0U if no set bit was found. */
+	/** Find last set bit. Returns InvalidIndex if no set bit was found. */
 	uint32 FindLastOne() const;
 
 	/** Returns the number of bits */
@@ -489,11 +492,11 @@ private:
 			{
 				const uint32 Index = CurrentBitIndex + FPlatformMath::CountTrailingZeros(~Word);
 				// Need to make sure the index is not out of bounds
-				return (Index < BitCount ? Index : ~0U);
+				return (Index < BitCount ? Index : FNetBitArrayBase::InvalidIndex);
 			}
 		}
 
-		return ~0U;
+		return FNetBitArrayBase::InvalidIndex;
 	}
 
 	static uint32 FindFirstOne(const StorageWordType* Storage, const uint32 WordCount, const uint32 BitCount)
@@ -505,11 +508,11 @@ private:
 			{
 				const uint32 Index = CurrentBitIndex + FPlatformMath::CountTrailingZeros(Word);
 				// Need to make sure the index is not out of bounds
-				return (Index < BitCount ? Index : ~0U);
+				return (Index < BitCount ? Index : FNetBitArrayBase::InvalidIndex);
 			}
 		}
 
-		return ~0U;
+		return FNetBitArrayBase::InvalidIndex;
 	}
 
 	static uint32 FindFirstZero(const StorageWordType* Storage, const uint32 WordCount, const uint32 BitCount, const uint32 StartIndex)
@@ -525,11 +528,11 @@ private:
 			{
 				const uint32 Index = CurrentBitIndex + FPlatformMath::CountTrailingZeros(~Word);
 				// Need to make sure the index is not out of bounds
-				return (Index < BitCount ? Index : ~0U);
+				return (Index < BitCount ? Index : FNetBitArrayBase::InvalidIndex);
 			}
 		}
 
-		return ~0U;
+		return FNetBitArrayBase::InvalidIndex;
 	}
 
 	static uint32 FindFirstOne(const StorageWordType* Storage, const uint32 WordCount, const uint32 BitCount, const uint32 StartIndex)
@@ -546,11 +549,11 @@ private:
 			{
 				const uint32 Index = CurrentBitIndex + FPlatformMath::CountTrailingZeros(Word);
 				// Need to make sure the index is not out of bounds
-				return (Index < BitCount ? Index : ~0U);
+				return (Index < BitCount ? Index : FNetBitArrayBase::InvalidIndex);
 			}
 		}
 
-		return ~0U;
+		return FNetBitArrayBase::InvalidIndex;
 	}
 
 	static uint32 FindLastZero(const StorageWordType* Storage, const uint32 WordCount, const uint32 BitCount)
@@ -571,7 +574,7 @@ private:
 			WordMask = StorageWordType(~StorageWordType(0));
 		}
 
-		return ~0U;
+		return FNetBitArrayBase::InvalidIndex;
 	}
 
 	static uint32 FindLastOne(const StorageWordType* Storage, const uint32 WordCount, const uint32 BitCount)
@@ -588,7 +591,7 @@ private:
 			}
 		}
 
-		return ~0U;
+		return FNetBitArrayBase::InvalidIndex;
 	}
 
 	NETCORE_API static uint32 GetSetBitIndices(const StorageWordType* Storage, const uint32 BitCount, const uint32 StartIndex, const uint32 Count, uint32* const OutIndices, const uint32 OutIndicesCapacity);
