Misc changes required to build consoles as modules

diff --git a/Engine/Source/Runtime/Core/Public/ProfilingDebugging/CallstackTrace.h b/Engine/Source/Runtime/Core/Public/ProfilingDebugging/CallstackTrace.h
--- a/Engine/Source/Runtime/Core/Public/ProfilingDebugging/CallstackTrace.h
+++ b/Engine/Source/Runtime/Core/Public/ProfilingDebugging/CallstackTrace.h
@@ -46,7 +46,7 @@ CORE_API uint32 CallstackTrace_GetCurrentId();
   */
 #define CALLSTACK_TRACE_LIMIT_CALLSTACKRESOLVE_SCOPE() 						FCallStackTraceLimitResolveScope PREPROCESSOR_JOIN(FCTLMScope,__LINE__)
 
-extern uint32 GCallStackTracingTlsSlotIndex;
+extern CORE_API uint32 GCallStackTracingTlsSlotIndex;
 
 /**
 * @return the fallback callstack address
diff --git a/Engine/Source/Runtime/Net/Core/Public/Net/Core/PropertyConditions/RepChangedPropertyTracker.h b/Engine/Source/Runtime/Net/Core/Public/Net/Core/PropertyConditions/RepChangedPropertyTracker.h
--- a/Engine/Source/Runtime/Net/Core/Public/Net/Core/PropertyConditions/RepChangedPropertyTracker.h
+++ b/Engine/Source/Runtime/Net/Core/Public/Net/Core/PropertyConditions/RepChangedPropertyTracker.h
@@ -13,7 +13,7 @@
  * TODO: This class (and arguably IRepChangedPropertyTracker) should be renamed to reflect
  *			what they actually do now.
  */
-class NETCORE_API FRepChangedPropertyTracker : public IRepChangedPropertyTracker
+class FRepChangedPropertyTracker : public IRepChangedPropertyTracker
 {
 public:
 	FRepChangedPropertyTracker() = delete;
