Disable receiving 'bHasPackageMapExports' serverside.

diff --git a/Engine/Source/Runtime/Engine/Private/NetConnection.cpp b/Engine/Source/Runtime/Engine/Private/NetConnection.cpp
--- a/Engine/Source/Runtime/Engine/Private/NetConnection.cpp
+++ b/Engine/Source/Runtime/Engine/Private/NetConnection.cpp
@@ -2989,8 +2989,18 @@ void UNetConnection::ReceivedPacket( FBitReader& Reader, bool bIsReinjectedPacke
 					return;
 				}
 
-				if ( Bunch.bHasPackageMapExports )
+				if (Bunch.bHasPackageMapExports)
 				{
+					// Clients still send NetGUID.IsDefault()/FExportFlags.bHasPath packets to the server, separate from this check
+					if (Driver->IsServer())
+					{
+						UE_LOG(LogNetTraffic, Error, TEXT("UNetConnection::ReceivedPacket: Server received bHasPackageMapExports packet."));
+
+						Close(ENetCloseResult::BunchServerPackageMapExports);
+
+						return;
+					}
+
 					Driver->NetGUIDInBytes += (BunchDataBits + (HeaderPos - IncomingStartPos)) >> 3;
 
 					if (IsInternalAck())
diff --git a/Engine/Source/Runtime/Net/Core/Private/Net/Core/Connection/NetCloseResult.cpp b/Engine/Source/Runtime/Net/Core/Private/Net/Core/Connection/NetCloseResult.cpp
--- a/Engine/Source/Runtime/Net/Core/Private/Net/Core/Connection/NetCloseResult.cpp
+++ b/Engine/Source/Runtime/Net/Core/Private/Net/Core/Connection/NetCloseResult.cpp
@@ -11,330 +11,27 @@
  * ENetCloseResult
  */
 
-// Avoid replacing with UEnum lookup - this is faster, and is used in constructor
-const TCHAR* LexToString(ENetCloseResult InResult)
-{
-	switch (InResult)
-	{
-	case ENetCloseResult::NetDriverAlreadyExists:
-		return TEXT("NetDriverAlreadyExists");
-
-	case ENetCloseResult::NetDriverCreateFailure:
-		return TEXT("NetDriverCreateFailure");
-
-	case ENetCloseResult::NetDriverListenFailure:
-		return TEXT("NetDriverListenFailure");
-
-	case ENetCloseResult::ConnectionLost:
-		return TEXT("ConnectionLost");
-
-	case ENetCloseResult::ConnectionTimeout:
-		return TEXT("ConnectionTimeout");
-
-	case ENetCloseResult::FailureReceived:
-		return TEXT("FailureReceived");
-
-	case ENetCloseResult::OutdatedClient:
-		return TEXT("OutdatedClient");
-
-	case ENetCloseResult::OutdatedServer:
-		return TEXT("OutdatedServer");
-
-	case ENetCloseResult::PendingConnectionFailure:
-		return TEXT("PendingConnectionFailure");
-
-	case ENetCloseResult::NetGuidMismatch:
-		return TEXT("NetGuidMismatch");
-
-	case ENetCloseResult::NetChecksumMismatch:
-		return TEXT("NetChecksumMismatch");
-
-	case ENetCloseResult::SecurityMalformedPacket:
-		return TEXT("SecurityMalformedPacket");
-
-	case ENetCloseResult::SecurityInvalidData:
-		return TEXT("SecurityInvalidData");
-
-	case ENetCloseResult::SecurityClosed:
-		return TEXT("SecurityClosed");
-
-	case ENetCloseResult::Unknown:
-		return TEXT("Unknown");
-
-	case ENetCloseResult::Success:
-		return TEXT("Success");
-
-	case ENetCloseResult::Extended:
-		return TEXT("Extended");
-
-	case ENetCloseResult::HostClosedConnection:
-		return TEXT("HostClosedConnection");
-
-	case ENetCloseResult::Disconnect:
-		return TEXT("Disconnect");
-
-	case ENetCloseResult::Upgrade:
-		return TEXT("Upgrade");
-
-	case ENetCloseResult::PreLoginFailure:
-		return TEXT("PreLoginFailure");
-
-	case ENetCloseResult::JoinFailure:
-		return TEXT("JoinFailure");
-
-	case ENetCloseResult::JoinSplitFailure:
-		return TEXT("JoinSplitFailure");
-
-	case ENetCloseResult::RPCDoS:
-		return TEXT("RPCDoS");
-
-	case ENetCloseResult::Cleanup:
-		return TEXT("Cleanup");
-
-	case ENetCloseResult::MissingLevelPackage:
-		return TEXT("MissingLevelPackage");
-
-	case ENetCloseResult::PacketHandlerIncomingError:
-		return TEXT("PacketHandlerIncomingError");
-
-	case ENetCloseResult::ZeroLastByte:
-		return TEXT("ZeroLastByte");
-
-	case ENetCloseResult::ZeroSize:
-		return TEXT("ZeroSize");
-
-	case ENetCloseResult::ReadHeaderFail:
-		return TEXT("ReadHeaderFail");
-
-	case ENetCloseResult::ReadHeaderExtraFail:
-		return TEXT("ReadHeaderExtraFail");
-
-	case ENetCloseResult::AckSequenceMismatch:
-		return TEXT("AckSequenceMismatch");
-
-	case ENetCloseResult::BunchBadChannelIndex:
-		return TEXT("BunchBadChannelIndex");
-
-	case ENetCloseResult::BunchChannelNameFail:
-		return TEXT("BunchChannelNameFail");
-
-	case ENetCloseResult::BunchWrongChannelType:
-		return TEXT("BunchWrongChannelType");
-
-	case ENetCloseResult::BunchHeaderOverflow:
-		return TEXT("BunchHeaderOverflow");
-
-	case ENetCloseResult::BunchDataOverflow:
-		return TEXT("BunchDataOverflow");
-
-	case ENetCloseResult::BunchPrematureControlChannel:
-		return TEXT("BunchPrematureControlChannel");
-
-	case ENetCloseResult::BunchPrematureChannel:
-		return TEXT("BunchPrematureChannel");
-
-	case ENetCloseResult::BunchPrematureControlClose:
-		return TEXT("BunchPrematureControlClose");
-
-	case ENetCloseResult::UnknownChannelType:
-		return TEXT("UnknownChannelType");
-
-	case ENetCloseResult::PrematureSend:
-		return TEXT("PrematureSend");
-
-	case ENetCloseResult::CorruptData:
-		return TEXT("CorruptData");
-
-	case ENetCloseResult::SocketSendFailure:
-		return TEXT("SocketSendFailure");
-
-	case ENetCloseResult::BadChildConnectionIndex:
-		return TEXT("BadChildConnectionIndex");
-
-	case ENetCloseResult::LogLimitInstant:
-		return TEXT("LogLimitInstant");
-
-	case ENetCloseResult::LogLimitSustained:
-		return TEXT("LogLimitSustained");
-
-	case ENetCloseResult::EncryptionFailure:
-		return TEXT("EncryptionFailure");
-
-	case ENetCloseResult::ReceivedNetGUIDBunchFail:
-		return TEXT("ReceivedNetGUIDBunchFail");
-
-	case ENetCloseResult::MaxReliableExceeded:
-		return TEXT("MaxReliableExceeded");
-
-	case ENetCloseResult::ReceivedNextBunchFail:
-		return TEXT("ReceivedNextBunchFail");
-
-	case ENetCloseResult::ReceivedNextBunchQueueFail:
-		return TEXT("ReceivedNextBunchQueueFail");
-
-	case ENetCloseResult::PartialInitialReliableDestroy:
-		return TEXT("PartialInitialReliableDestroy");
-
-	case ENetCloseResult::PartialMergeReliableDestroy:
-		return TEXT("PartialMergeReliableDestroy");
-
-	case ENetCloseResult::PartialInitialNonByteAligned:
-		return TEXT("PartialInitialNonByteAligned");
-
-	case ENetCloseResult::PartialNonByteAligned:
-		return TEXT("PartialNonByteAligned");
-
-	case ENetCloseResult::PartialFinalPackageMapExports:
-		return TEXT("PartialFinalPackageMapExports");
-
-	case ENetCloseResult::PartialTooLarge:
-		return TEXT("PartialTooLarge");
-
-	case ENetCloseResult::AlreadyOpen:
-		return TEXT("AlreadyOpen");
-
-	case ENetCloseResult::ReliableBeforeOpen:
-		return TEXT("ReliableBeforeOpen");
-
-	case ENetCloseResult::ReliableBufferOverflow:
-		return TEXT("ReliableBufferOverflow");
-
-	case ENetCloseResult::RPCReliableBufferOverflow:
-		return TEXT("RPCReliableBufferOverflow");
-
-	case ENetCloseResult::ControlChannelClose:
-		return TEXT("ControlChannelClose");
-
-	case ENetCloseResult::ControlChannelEndianCheck:
-		return TEXT("ControlChannelEndianCheck");
-
-	case ENetCloseResult::ControlChannelPlayerChannelFail:
-		return TEXT("ControlChannelPlayerChannelFail");
-
-	case ENetCloseResult::ControlChannelMessageUnknown:
-		return TEXT("ControlChannelMessageUnknown");
-
-	case ENetCloseResult::ControlChannelMessageFail:
-		return TEXT("ControlChannelMessageFail");
-
-	case ENetCloseResult::ControlChannelMessagePayloadFail:
-		return TEXT("ControlChannelMessagePayloadFail");
-
-	case ENetCloseResult::ControlChannelBunchOverflowed:
-		return TEXT("ControlChannelBunchOverflowed");
-
-	case ENetCloseResult::ControlChannelQueueBunchOverflowed:
-		return TEXT("ControlChannelQueueBunchOverflowed");
-
-	case ENetCloseResult::ClientHasMustBeMappedGUIDs:
-		return TEXT("ClientHasMustBeMappedGUIDs");
-
-	case ENetCloseResult::UnregisteredMustBeMappedGUID:
-		return TEXT("UnregisteredMustBeMappedGUID");
-
-	case ENetCloseResult::ObjectReplicatorReceivedBunchFail:
-		return TEXT("ObjectReplicatorReceivedBunchFail");
-
-	case ENetCloseResult::ContentBlockFail:
-		return TEXT("ContentBlockFail");
-
-	case ENetCloseResult::ContentBlockHeaderRepLayoutFail:
-		return TEXT("ContentBlockHeaderRepLayoutFail");
-
-	case ENetCloseResult::ContentBlockHeaderIsActorFail:
-		return TEXT("ContentBlockHeaderIsActorFail");
-
-	case ENetCloseResult::ContentBlockHeaderObjFail:
-		return TEXT("ContentBlockHeaderObjFail");
-
-	case ENetCloseResult::ContentBlockHeaderPrematureEnd:
-		return TEXT("ContentBlockHeaderPrematureEnd");
-
-	case ENetCloseResult::ContentBlockHeaderSubObjectActor:
-		return TEXT("ContentBlockHeaderSubObjectActor");
-
-	case ENetCloseResult::ContentBlockHeaderBadParent:
-		return TEXT("ContentBlockHeaderBadParent");
-
-	case ENetCloseResult::ContentBlockHeaderInvalidCreate:
-		return TEXT("ContentBlockHeaderInvalidCreate");
-
-	case ENetCloseResult::ContentBlockHeaderStablyNamedFail:
-		return TEXT("ContentBlockHeaderStablyNamedFail");
-
-	case ENetCloseResult::ContentBlockHeaderNoSubObjectClass:
-		return TEXT("ContentBlockHeaderNoSubObjectClass");
-
-	case ENetCloseResult::ContentBlockHeaderUObjectSubObject:
-		return TEXT("ContentBlockHeaderUObjectSubObject");
-
-	case ENetCloseResult::ContentBlockHeaderAActorSubObject:
-		return TEXT("ContentBlockHeaderAActorSubObject");
-
-	case ENetCloseResult::ContentBlockHeaderFail:
-		return TEXT("ContentBlockHeaderFail");
-
-	case ENetCloseResult::ContentBlockPayloadBitsFail:
-		return TEXT("ContentBlockPayloadBitsFail");
-
-	case ENetCloseResult::FieldHeaderRepIndex:
-		return TEXT("FieldHeaderRepIndex");
-
-	case ENetCloseResult::FieldHeaderBadRepIndex:
-		return TEXT("FieldHeaderBadRepIndex");
-
-	case ENetCloseResult::FieldHeaderPayloadBitsFail:
-		return TEXT("FieldHeaderPayloadBitsFail");
-
-	case ENetCloseResult::FieldPayloadFail:
-		return TEXT("FieldPayloadFail");
-
-	case ENetCloseResult::BeaconControlFlowError:
-		return TEXT("BeaconControlFlowError");
-
-	case ENetCloseResult::BeaconAuthenticationFailure:
-		return TEXT("BeaconAuthenticationFailure");
-
-	case ENetCloseResult::BeaconUnableToParsePacket:
-		return TEXT("BeaconUnableToParsePacket");
-
-	case ENetCloseResult::BeaconLoginInvalidIdError:
-		return TEXT("BeaconLoginInvalidIdError");
-
-	case ENetCloseResult::BeaconLoginInvalidAuthHandlerError:
-		return TEXT("BeaconLoginInvalidAuthHandlerError");
-
-	case ENetCloseResult::BeaconAuthError:
-		return TEXT("BeaconAuthError");
-
-	case ENetCloseResult::BeaconSpawnClientWorldPackageNameError:
-		return TEXT("BeaconSpawnClientWorldPackageNameError");
-
-	case ENetCloseResult::BeaconSpawnExistingActorError:
-		return TEXT("BeaconSpawnExistingActorError");
-
-	case ENetCloseResult::BeaconSpawnFailureError:
-		return TEXT("BeaconSpawnFailureError");
-
-	case ENetCloseResult::BeaconSpawnNetGUIDAckNoActor:
-		return TEXT("BeaconSpawnNetGUIDAckNoActor");
-
-	case ENetCloseResult::BeaconSpawnNetGUIDAckNoHost:
-		return TEXT("BeaconSpawnNetGUIDAckNoHost");
+#ifndef CASE_ENUM_TO_TEXT_RET
+#define CASE_ENUM_TO_TEXT_RET(txt) case txt: ReturnVal = TEXT(#txt); break;
+#endif
 
-	case ENetCloseResult::BeaconSpawnUnexpectedError:
-		return TEXT("BeaconSpawnUnexpectedError");
+const TCHAR* LexToString(ENetCloseResult Enum)
+{
+	const TCHAR* ReturnVal = TEXT("::Invalid");
 
-	case ENetCloseResult::FaultDisconnect:
-		return TEXT("FaultDisconnect");
+	switch (Enum)
+	{
+		FOREACH_ENUM_ENETCLOSERESULT(CASE_ENUM_TO_TEXT_RET)
+	}
 
-	case ENetCloseResult::NotRecoverable:
-		return TEXT("NotRecoverable");
+	while (*ReturnVal != ':')
+	{
+		ReturnVal++;
+	}
 
+	ReturnVal += 2;
 
-	default:
-		return TEXT("Invalid");
-	}
+	return ReturnVal;
 }
 
 ENetCloseResult FromNetworkFailure(ENetworkFailure::Type Val)
@@ -475,24 +172,6 @@ bool FNetCloseResultEnumTest::RunTest(const FString& Parameters)
 		}
 	}
 
-
-	if (TestTrue(TEXT("ENetCloseResult must exist"), NetCloseResultEnum != nullptr) && NetCloseResultEnum != nullptr)
-	{
-		const int64 NetCloseResultEnumLast = NetCloseResultEnum->GetMaxEnumValue() - 1;
-		bool bLexMismatch = false;
-
-		for (int64 EnumIdx=0; EnumIdx<=NetCloseResultEnumLast; EnumIdx++)
-		{
-			if (NetCloseResultEnum->GetNameStringByValue(EnumIdx) != LexToString((ENetCloseResult)EnumIdx))
-			{
-				bLexMismatch = true;
-				break;
-			}
-		}
-
-		TestFalse(TEXT("ENetCloseResult must not be missing LexToString entries"), bLexMismatch);
-	}
-
 	return true;
 }
 #endif
diff --git a/Engine/Source/Runtime/Net/Core/Public/Net/Core/Connection/NetCloseResult.h b/Engine/Source/Runtime/Net/Core/Public/Net/Core/Connection/NetCloseResult.h
--- a/Engine/Source/Runtime/Net/Core/Public/Net/Core/Connection/NetCloseResult.h
+++ b/Engine/Source/Runtime/Net/Core/Public/Net/Core/Connection/NetCloseResult.h
@@ -7,6 +7,8 @@
 #include "Net/Core/Connection/NetResult.h"
 #include "Templates/PimplPtr.h"
 
+#include "NetCloseResult.generated.h"
+
 
 /**
  * Network Close results (can occur before NetDriver/NetConnection creation, and can be success or error)
@@ -144,6 +146,9 @@ enum class ENetCloseResult : uint32
 	/** Bunch data serialization overflowed */
 	BunchDataOverflow,
 
+	/** Server received bHasPackageMapExports packet */
+	BunchServerPackageMapExports,
+
 	/** Received control channel bunch before control channel was created */
 	BunchPrematureControlChannel,
 
@@ -370,7 +375,7 @@ enum class ENetCloseResult : uint32
 DECLARE_NETRESULT_ENUM(ENetCloseResult);
 
 
-NETCORE_API const TCHAR* LexToString(ENetCloseResult InResult);
+NETCORE_API const TCHAR* LexToString(ENetCloseResult Enum);
 
 /** Converts from ENetworkFailure to ENetCloseResult */
 NETCORE_API ENetCloseResult FromNetworkFailure(ENetworkFailure::Type Val);
