Fix non-unity error

diff --git a/Engine/Source/Runtime/Net/Core/Private/Net/Core/Misc/NetConditionGroupManager.cpp b/Engine/Source/Runtime/Net/Core/Private/Net/Core/Misc/NetConditionGroupManager.cpp
--- a/Engine/Source/Runtime/Net/Core/Private/Net/Core/Misc/NetConditionGroupManager.cpp
+++ b/Engine/Source/Runtime/Net/Core/Private/Net/Core/Misc/NetConditionGroupManager.cpp
@@ -1,6 +1,7 @@
 // Copyright Epic Games, Inc. All Rights Reserved.
 
 #include "Net/Core/Misc/NetConditionGroupManager.h"
+#include "UObject/Object.h"
 
 extern const FName UE::Net::NetGroupOwner(TEXT("NetGroupOwner"));
 extern const FName UE::Net::NetGroupReplay(TEXT("NetGroupReplay"));
