SubObject
* Added define to setup a project to use the subobject list by default without having to use the config.

diff --git a/Engine/Source/Runtime/Engine/Private/Actor.cpp b/Engine/Source/Runtime/Engine/Private/Actor.cpp
--- a/Engine/Source/Runtime/Engine/Private/Actor.cpp
+++ b/Engine/Source/Runtime/Engine/Private/Actor.cpp
@@ -126,6 +126,7 @@ void AActor::InitializeDefaults()
 	bReplicates = false;
 	bCallPreReplication = true;
 	bCallPreReplicationForReplay = true;
+	bReplicateUsingRegisteredSubObjectList = DEFAULT_USE_SUBOBJECT_LIST;
 	NetPriority = 1.0f;
 	NetUpdateFrequency = 100.0f;
 	MinNetUpdateFrequency = 2.0f;
diff --git a/Engine/Source/Runtime/Engine/Private/Components/ActorComponent.cpp b/Engine/Source/Runtime/Engine/Private/Components/ActorComponent.cpp
--- a/Engine/Source/Runtime/Engine/Private/Components/ActorComponent.cpp
+++ b/Engine/Source/Runtime/Engine/Private/Components/ActorComponent.cpp
@@ -332,6 +332,8 @@ UActorComponent::UActorComponent(const FObjectInitializer& ObjectInitializer /*=
 
 	bMarkedForPreEndOfFrameSync = false;
 	bAsyncPhysicsTickEnabled = false;
+
+	bReplicateUsingRegisteredSubObjectList = DEFAULT_USE_SUBOBJECT_LIST;
 }
 
 void UActorComponent::PostInitProperties()
diff --git a/Engine/Source/Runtime/Net/Core/Public/Net/Core/Misc/NetSubObjectRegistry.h b/Engine/Source/Runtime/Net/Core/Public/Net/Core/Misc/NetSubObjectRegistry.h
--- a/Engine/Source/Runtime/Net/Core/Public/Net/Core/Misc/NetSubObjectRegistry.h
+++ b/Engine/Source/Runtime/Net/Core/Public/Net/Core/Misc/NetSubObjectRegistry.h
@@ -7,6 +7,15 @@
 
 class UActorComponent;
 
+/** 
+* Chooses the default value of Actor and ActorComponent's bReplicateUsingRegisteredSubObjectList.  
+* If a project wants to replicate subobjects using the registration list by default it should define it to true in it's TargetRules.
+* Classes can also chose to use the list or not by setting bReplicateUsingRegisteredSubObjectList in their config or constructor.
+*/
+#ifndef DEFAULT_USE_SUBOBJECT_LIST
+	#define DEFAULT_USE_SUBOBJECT_LIST false
+#endif
+
 namespace UE::Net 
 {
 
