Fixed missing call to PostReplicatedChange on the fast array serializer, when mapping previously unmapped objects (Items PostReplicatedChange calls were not affected).

Fixed for both 'net.SupportFastArrayDelta' on/off paths.

diff --git a/Engine/Source/Runtime/Engine/Private/RepLayout.cpp b/Engine/Source/Runtime/Engine/Private/RepLayout.cpp
--- a/Engine/Source/Runtime/Engine/Private/RepLayout.cpp
+++ b/Engine/Source/Runtime/Engine/Private/RepLayout.cpp
@@ -8056,6 +8056,11 @@ void FRepLayout::UpdateUnmappedGuidsForFastArray(FFastArrayDeltaSerializeParams&
 
 			if (bOutSomeObjectsWereMapped)
 			{
+				if (Params.ReadChangedElements != nullptr)
+				{
+					Params.ReadChangedElements->Add(ItemIndex);
+				}
+
 				Params.PostReplicatedChange(ElementData, Params);
 			}
 
diff --git a/Engine/Source/Runtime/Net/Core/Classes/Net/Serialization/FastArraySerializer.h b/Engine/Source/Runtime/Net/Core/Classes/Net/Serialization/FastArraySerializer.h
--- a/Engine/Source/Runtime/Net/Core/Classes/Net/Serialization/FastArraySerializer.h
+++ b/Engine/Source/Runtime/Net/Core/Classes/Net/Serialization/FastArraySerializer.h
@@ -1168,6 +1168,8 @@ bool FFastArraySerializer::FastArrayDeltaSerialize(TArray<Type> &Items, FNetDelt
 	{
 		UE_LOG(LogNetFastTArray, Log, TEXT("FastArrayDeltaSerialize for %s. %s. UpdateUnmappedObjects"), *InnerStruct->GetName(), *InnerStruct->GetOwnerStruct()->GetName());
 
+		TArray<int32, TInlineAllocator<8>> ChangedIndices;
+
 		// Loop over each item that has unmapped objects
 		for ( auto It = ArraySerializer.GuidReferencesMap.CreateIterator(); It; ++It )
 		{
@@ -1225,7 +1227,10 @@ bool FFastArraySerializer::FastArrayDeltaSerialize(TArray<Type> &Items, FNetDelt
 					Parms.bCalledPreNetReceive = true;
 				}
 
-				Type* ThisElement = &Items[ArraySerializer.ItemMap.FindChecked( ElementID )];
+				const int32 ElementIndex = ArraySerializer.ItemMap.FindChecked(ElementID);
+				Type* ThisElement = &Items[ElementIndex];
+
+				ChangedIndices.Add(ElementIndex);
 
 				// Initialize the reader with the stored buffer that we need to read from
 				FNetBitReader Reader( Parms.Map, GuidReferences.Buffer.GetData(), GuidReferences.NumBufferBits );
@@ -1255,6 +1260,7 @@ bool FFastArraySerializer::FastArrayDeltaSerialize(TArray<Type> &Items, FNetDelt
 
 		if (Parms.bOutSomeObjectsWereMapped)
 		{
+			ArraySerializer.PostReplicatedChange(ChangedIndices, Items.Num());
 			Helper.CallPostReplicatedReceiveOrNot();
 		}
 		return true;
@@ -1576,10 +1582,16 @@ bool FFastArraySerializer::FastArrayDeltaSerialize_DeltaSerializeStructs(TArray<
 	else if (Parms.bUpdateUnmappedObjects)
 	{
 		UE_LOG(LogNetFastTArray, Log, TEXT("FastArrayDeltaSerialize_DeltaSerializeStruct for %s. %s. UpdateUnmappedObjects"), *InnerStruct->GetName(), *InnerStruct->GetOwnerStruct()->GetName());
+
+		TArray<int32, TInlineAllocator<8>> ChangedIndices;
+
+		DeltaSerializeParams.ReadChangedElements = &ChangedIndices;
+
 		Parms.NetSerializeCB->UpdateUnmappedGuidsForFastArray(DeltaSerializeParams);
 
 		if (Parms.bOutSomeObjectsWereMapped)
 		{
+			ArraySerializer.PostReplicatedChange(ChangedIndices, Items.Num());
 			Helper.CallPostReplicatedReceiveOrNot();
 		}
 
