Iris experimental replication system.

diff --git a/Engine/Source/Runtime/Net/Core/Private/Net/Core/Trace/Reporters/NetTraceReporter.h b/Engine/Source/Runtime/Net/Core/Private/Net/Core/Trace/Reporters/NetTraceReporter.h
--- a/Engine/Source/Runtime/Net/Core/Private/Net/Core/Trace/Reporters/NetTraceReporter.h
+++ b/Engine/Source/Runtime/Net/Core/Private/Net/Core/Trace/Reporters/NetTraceReporter.h
@@ -14,15 +14,15 @@ struct FNetTraceReporter
 	static void ReportPacketContent(FNetTracePacketContentEvent* Events, uint32 EventCount, const FNetTracePacketInfo& PacketInfo);
 	static void ReportPacketDropped(const FNetTracePacketInfo& PacketInfo);
 	static void ReportPacket(const FNetTracePacketInfo& PacketInfo, uint32 PacketBits);
-	static void ReportAnsiName(FNetDebugNameId NameId, uint32 NameSize, const char* Name);	
-	static void ReportObjectCreated(uint32 GameInstanceId, uint32 NetObjectId, FNetDebugNameId NameId, uint64 TypeIdentifier, uint32 OwnerId);
+	static void ReportAnsiName(UE::Net::FNetDebugNameId NameId, uint32 NameSize, const char* Name);	
+	static void ReportObjectCreated(uint32 GameInstanceId, uint32 NetObjectId, UE::Net::FNetDebugNameId NameId, uint64 TypeIdentifier, uint32 OwnerId);
 	static void ReportObjectDestroyed(uint32 GameInstanceId, uint32 NetObjectId);
 	static void ReportConnectionCreated(uint32 GameInstanceId, uint32 ConnectionId);
 	static void ReportConnectionStateUpdated(uint32 GameInstanceId, uint32 ConnectionId, uint8 ConnectionStateValue);
 	static void ReportConnectionUpdated(uint32 GameInstanceId, uint32 ConnectionId, const TCHAR* AddressString, const TCHAR* OwningActor);
 	static void ReportConnectionClosed(uint32 GameInstanceId, uint32 ConnectionId);
-	static void ReportPacketStatsCounter(uint32 GameInstanceId, uint32 ConnectionId, FNetDebugNameId CounterNameId, uint32 StatValue);
-	static void ReportFrameStatsCounter(uint32 GameInstanceId, FNetDebugNameId CounterNameId, uint32 StatValue);
+	static void ReportPacketStatsCounter(uint32 GameInstanceId, uint32 ConnectionId, UE::Net::FNetDebugNameId CounterNameId, uint32 StatValue);
+	static void ReportFrameStatsCounter(uint32 GameInstanceId, UE::Net::FNetDebugNameId CounterNameId, uint32 StatValue);
 	static void ReportInstanceUpdated(uint32 GameInstanceId, bool bIsServer, const TCHAR* Name);
 	static void ReportInstanceDestroyed(uint32 GameInstanceId);
 };
diff --git a/Engine/Source/Runtime/Net/Core/Public/Net/Core/PropertyConditions/RepChangedPropertyTracker.h b/Engine/Source/Runtime/Net/Core/Public/Net/Core/PropertyConditions/RepChangedPropertyTracker.h
--- a/Engine/Source/Runtime/Net/Core/Public/Net/Core/PropertyConditions/RepChangedPropertyTracker.h
+++ b/Engine/Source/Runtime/Net/Core/Public/Net/Core/PropertyConditions/RepChangedPropertyTracker.h
@@ -4,6 +4,7 @@
 
 #include "CoreMinimal.h"
 #include "UObject/CoreNet.h"
+#include "Delegates/Delegate.h"
 
 /**
  * This class is used to store meta data about properties that is shared between connections,
@@ -77,3 +78,12 @@ public:
 	UE_DEPRECATED(5.0, "No longer used, see UReplaySubsystem::SetExternalDataForObject")
 	uint32 ExternalDataNumBits;
 };
+
+namespace UE::Net::Private
+{
+#if UE_WITH_IRIS
+DECLARE_DELEGATE_ThreeParams(FIrisSetPropertyCustomCondition, const UObject*, const uint16, const bool);
+NETCORE_API void SetIrisSetPropertyCustomConditionDelegate(const FIrisSetPropertyCustomCondition& Delegate);
+#endif
+}
+
diff --git a/Engine/Source/Runtime/Net/Core/Public/Net/Core/PushModel/PushModel.h b/Engine/Source/Runtime/Net/Core/Public/Net/Core/PushModel/PushModel.h
--- a/Engine/Source/Runtime/Net/Core/Public/Net/Core/PushModel/PushModel.h
+++ b/Engine/Source/Runtime/Net/Core/Public/Net/Core/PushModel/PushModel.h
@@ -7,6 +7,7 @@
 #if WITH_PUSH_MODEL
 
 #include "CoreMinimal.h"
+#include "Delegates/Delegate.h"
 #include "UObject/ObjectKey.h"
 
 /**
@@ -279,13 +280,32 @@ namespace UEPushModelPrivate
 	//~ Using int32 isn't very forward looking, but for now GUObjectArray also uses int32
 	//~ so we're probably safe.
 
-	using FNetPushObjectId = int32;
 	using FNetPushPerNetDriverId = int32;
+	using FNetLegacyPushObjectId = int32;
+	using FNetIrisPushObjectId = uint64;
+
+	// The ID need to be able to handle both Iris and non-Iris replication.
+	struct FNetPushObjectId
+	{
+		FNetPushObjectId() : Value(0xFFFFFFFFFFFFFFFFULL) {}
+		explicit FNetPushObjectId(FNetLegacyPushObjectId Id) { Value = 0xFFFFFFFF00000000ULL | uint32(Id); }
+		explicit FNetPushObjectId(FNetIrisPushObjectId Id) { Value = Id; }
+
+		bool IsValid() const { return Value != 0xFFFFFFFFFFFFFFFFULL; }
+		bool IsIrisId() const { return uint32(Value >> 32U) != 0xFFFFFFFFU; }
+		uint64 GetValue() const { return Value; }
+
+		FNetLegacyPushObjectId GetLegacyPushObjectId() const { return FNetLegacyPushObjectId(Value & 0xFFFFFFFFU);}
+		FNetIrisPushObjectId GetIrisPushObjectId() const { return Value;}
+
+	private:
+		uint64 Value;
+	};
 
 	/** A handle that can be used to refer to a Push Model Object State owned by a specific Net Driver. */
 	struct FPushModelPerNetDriverHandle
 	{
-		FPushModelPerNetDriverHandle(const FNetPushPerNetDriverId InNetDriverId, const FNetPushObjectId InObjectId)
+		FPushModelPerNetDriverHandle(const FNetPushPerNetDriverId InNetDriverId, const FNetLegacyPushObjectId InObjectId)
 			: NetDriverId(InNetDriverId)
 			, ObjectId(InObjectId)
 		{
@@ -300,7 +320,7 @@ namespace UEPushModelPrivate
 		const FNetPushPerNetDriverId NetDriverId;
 		
 		//! The "globally unique" ID that is used to refer to the Push Model Object.
-		const FNetPushObjectId ObjectId;
+		const FNetLegacyPushObjectId ObjectId;
 
 		const bool IsValid() const
 		{
@@ -330,7 +350,7 @@ namespace UEPushModelPrivate
 
 	static FString ToString(const FNetPushObjectId Id)
 	{
-		return FString::Printf(TEXT("%d"), Id);
+		return FString::Printf(TEXT("0x%" UINT64_X_FMT), Id.GetValue());
 	}
 
 	/**
@@ -346,9 +366,10 @@ namespace UEPushModelPrivate
 	 * constants instead of relying on UProperty::RepIndex.
 	 *
 	 * @param Object	The object we're marking dirty.
+	 * @param ObjectId	The push ID for the object we're marking dirty.
 	 * @param RepIndex	The index of the property we're marking dirty. UProperty::RepIndex.
 	 */
-	NETCORE_API void MarkPropertyDirty(const FNetPushObjectId ObjectId, const int32 RepIndex);
+	NETCORE_API void MarkPropertyDirty(const UObject* Object, const FNetPushObjectId ObjectId, const int32 RepIndex);
 
 	/**
 	 * Marks a range of properties as dirty, causing it to be compared the next time the owner is considered
@@ -362,9 +383,10 @@ namespace UEPushModelPrivate
 	 * constants instead of relying on UProperty::RepIndex.
 	 *
 	 * @param Object	The object we're marking dirty.
+	 * @param ObjectId	The push ID for the object we're marking dirty.
 	 * @param RepIndex	The index of the property we're marking dirty. UProperty::RepIndex.
 	 */
-	NETCORE_API void MarkPropertyDirty(const FNetPushObjectId ObjectId, const int32 StartRepIndex, const int32 EndRepIndex);
+	NETCORE_API void MarkPropertyDirty(const UObject* Object, const FNetPushObjectId ObjectId, const int32 StartRepIndex, const int32 EndRepIndex);
 
 	//~ As the comments above state, none of the methods in this namespace should be invoked directly.
 	//~ Particularly, the methods below are **only** needed by internal systems.
@@ -378,9 +400,18 @@ namespace UEPushModelPrivate
 
 	NETCORE_API bool DoesHaveDirtyPropertiesOrRecentlyCollectedGarbage(const FPushModelPerNetDriverHandle Handle);
 
-	NETCORE_API bool ValidateObjectIdReassignment(FNetPushObjectId CurrentId, FNetPushObjectId NewId);
+	NETCORE_API bool ValidateObjectIdReassignment(FNetLegacyPushObjectId CurrentId, FNetLegacyPushObjectId NewId);
 
 	NETCORE_API void LogMemory(FOutputDevice& Ar);
+
+#if UE_WITH_IRIS
+	// For internal use only.
+	DECLARE_DELEGATE_ThreeParams(FIrisMarkPropertyDirty, const UObject*, UEPushModelPrivate::FNetIrisPushObjectId, const int32);
+	DECLARE_DELEGATE_FourParams(FIrisMarkPropertiesDirty, const UObject*, UEPushModelPrivate::FNetIrisPushObjectId, const int32, const int32);
+
+	NETCORE_API void SetIrisMarkPropertyDirtyDelegate(const FIrisMarkPropertyDirty& Delegate);
+	NETCORE_API void SetIrisMarkPropertiesDirtyDelegate(const FIrisMarkPropertiesDirty& Delegate);
+#endif // UE_WITH_IRIS
 }
 
 
@@ -395,34 +426,34 @@ namespace UEPushModelPrivate
 
 #define IS_PROPERTY_REPLICATED(Property) (0 != (EPropertyFlags::CPF_Net & Property->PropertyFlags))
 
-#define CONDITIONAL_ON_OBJECT_NET_ID(Object, Work) { const UEPushModelPrivate::FNetPushObjectId PrivatePushId = Object->GetNetPushId(); Work; }
-#define CONDITIONAL_ON_OBJECT_NET_ID_DYNAMIC(Object, Work) { const UEPushModelPrivate::FNetPushObjectId PrivatePushId = Object->GetNetPushIdDynamic(); Work; }
-#define CONDITIONAL_ON_REP_INDEX_AND_OBJECT_NET_ID(Object, Property, Work) if (IS_PROPERTY_REPLICATED(Property)) { const UEPushModelPrivate::FNetPushObjectId PrivatePushId = Object->GetNetPushIdDynamic();  Work; }
+#define CONDITIONAL_ON_OBJECT_NET_ID(Object, Work) { const UEPushModelPrivate::FNetPushObjectId PrivatePushId(Object->GetNetPushId()); Work; }
+#define CONDITIONAL_ON_OBJECT_NET_ID_DYNAMIC(Object, Work) { const UEPushModelPrivate::FNetPushObjectId PrivatePushId(Object->GetNetPushIdDynamic()); Work; }
+#define CONDITIONAL_ON_REP_INDEX_AND_OBJECT_NET_ID(Object, Property, Work) if (IS_PROPERTY_REPLICATED(Property)) { const UEPushModelPrivate::FNetPushObjectId PrivatePushId(Object->GetNetPushIdDynamic()); Work; }
 
 //~ For these macros, we won't bother checking if Push Model is enabled. Instead, we'll just check to see whether or not the Custom ID is valid.
 
 // Marks a property dirty by RepIndex without doing additional rep index validation.
-#define MARK_PROPERTY_DIRTY_UNSAFE(Object, RepIndex) CONDITIONAL_ON_OBJECT_NET_ID_DYNAMIC(Object, UEPushModelPrivate::MarkPropertyDirty(PrivatePushId, RepIndex))
+#define MARK_PROPERTY_DIRTY_UNSAFE(Object, RepIndex) CONDITIONAL_ON_OBJECT_NET_ID_DYNAMIC(Object, UEPushModelPrivate::MarkPropertyDirty(Object, PrivatePushId, RepIndex))
 
 // Marks a property dirty by UProperty*, validating that it's actually a replicated property.
-#define MARK_PROPERTY_DIRTY(Object, Property) CONDITIONAL_ON_REP_INDEX_AND_OBJECT_NET_ID(Object, Property, UEPushModelPrivate::MarkPropertyDirty(PrivatePushId, Property->RepIndex))
+#define MARK_PROPERTY_DIRTY(Object, Property) CONDITIONAL_ON_REP_INDEX_AND_OBJECT_NET_ID(Object, Property, UEPushModelPrivate::MarkPropertyDirty(Object, PrivatePushId, Property->RepIndex))
 
 
 // Marks a static array property dirty given, the Object, UProperty*, and Index.
-#define MARK_PROPERTY_DIRTY_STATIC_ARRAY_INDEX(Object, Property, ArrayIndex) CONDITIONAL_ON_REP_INDEX_AND_OBJECT_NET_ID(Object, Property, UEPushModelPrivate::MarkPropertyDirty(PrivatePushId, Property->RepIndex + ArrayIndex))
+#define MARK_PROPERTY_DIRTY_STATIC_ARRAY_INDEX(Object, Property, ArrayIndex) CONDITIONAL_ON_REP_INDEX_AND_OBJECT_NET_ID(Object, Property, UEPushModelPrivate::MarkPropertyDirty(Object, PrivatePushId, Property->RepIndex + ArrayIndex))
 
 // Marks all elements of a static array property dirty, given the Object and UProperty*
-#define MARK_PROPERTY_DIRTY_STATIC_ARRAY(Object, Property) CONDITIONAL_ON_REP_INDEX_AND_OBJECT_NET_ID(Object, Property, UEPushModelPrivate::MarkPropertyDirty(PrivatePushId, Property->RepIndex, Property->RepIndex + Property->ArrayDim - 1))
+#define MARK_PROPERTY_DIRTY_STATIC_ARRAY(Object, Property) CONDITIONAL_ON_REP_INDEX_AND_OBJECT_NET_ID(Object, Property, UEPushModelPrivate::MarkPropertyDirty(Object, PrivatePushId, Property->RepIndex, Property->RepIndex + Property->ArrayDim - 1))
 
 
 // Marks a property dirty, given the Class Name, Property Name, and Object. This will fail to compile if the Property or Class aren't valid.
-#define MARK_PROPERTY_DIRTY_FROM_NAME(ClassName, PropertyName, Object) CONDITIONAL_ON_OBJECT_NET_ID(Object, UEPushModelPrivate::MarkPropertyDirty(PrivatePushId, GET_PROPERTY_REP_INDEX(ClassName, PropertyName)))
+#define MARK_PROPERTY_DIRTY_FROM_NAME(ClassName, PropertyName, Object) CONDITIONAL_ON_OBJECT_NET_ID(Object, UEPushModelPrivate::MarkPropertyDirty(Object, PrivatePushId, GET_PROPERTY_REP_INDEX(ClassName, PropertyName)))
 
 // Marks a static array property dirty, given the Class Name, Property Name, Index, and Object. This will fail to compile if the Property and Class aren't valid. Callers are responsible for validating the index.
-#define MARK_PROPERTY_DIRTY_FROM_NAME_STATIC_ARRAY_INDEX(ClassName, PropertyName, ArrayIndex, Object) CONDITIONAL_ON_OBJECT_NET_ID(Object, UEPushModelPrivate::MarkPropertyDirty(PrivatePushId, GET_PROPERTY_REP_INDEX_STATIC_ARRAY_INDEX(ClassName, PropertyName, ArrayIndex)))
+#define MARK_PROPERTY_DIRTY_FROM_NAME_STATIC_ARRAY_INDEX(ClassName, PropertyName, ArrayIndex, Object) CONDITIONAL_ON_OBJECT_NET_ID(Object, UEPushModelPrivate::MarkPropertyDirty(Object, PrivatePushId, GET_PROPERTY_REP_INDEX_STATIC_ARRAY_INDEX(ClassName, PropertyName, ArrayIndex)))
 
 // Marks an entire static array property dirty, given the Class Name, Property Name, and Object. This will fail to compile if the Property or Class aren't valid.
-#define MARK_PROPERTY_DIRTY_FROM_NAME_STATIC_ARRAY(ClassName, PropertyName, Object) CONDITIONAL_ON_OBJECT_NET_ID(Object, UEPushModelPrivate::MarkPropertyDirty(PrivatePushId, GET_PROPERTY_REP_INDEX_STATIC_ARRAY_START(ClassName, PropertyName), GET_PROPERTY_REP_INDEX_STATIC_ARRAY_END(ClassName, PropertyName))
+#define MARK_PROPERTY_DIRTY_FROM_NAME_STATIC_ARRAY(ClassName, PropertyName, Object) CONDITIONAL_ON_OBJECT_NET_ID(Object, UEPushModelPrivate::MarkPropertyDirty(Object, PrivatePushId, GET_PROPERTY_REP_INDEX_STATIC_ARRAY_START(ClassName, PropertyName), GET_PROPERTY_REP_INDEX_STATIC_ARRAY_END(ClassName, PropertyName))
 
 /**
  * This is used to reduce lines of code (mostly in setters) by doing the comparison check before changing the value and marking dirty only if we have changed the value. 
@@ -431,6 +462,11 @@ namespace UEPushModelPrivate
 #define COMPARE_ASSIGN_AND_MARK_PROPERTY_DIRTY(ClassName, PropertyName, NewValue, Object) \
 	if (NewValue != PropertyName) { PropertyName = NewValue; MARK_PROPERTY_DIRTY_FROM_NAME(ClassName, PropertyName, Object); }
 
+#if UE_WITH_IRIS
+#define UE_NET_SET_IRIS_MARK_PROPERTY_DIRTY_DELEGATE(Delegate) UEPushModelPrivate::SetIrisMarkPropertyDirtyDelegate(Delegate)
+#define UE_NET_SET_IRIS_MARK_PROPERTIES_DIRTY_DELEGATE(Delegate) UEPushModelPrivate::SetIrisMarkPropertiesDirtyDelegate(Delegate)
+#endif // UE_WITH_IRIS
+
 #else // WITH_PUSH_MODEL
 
 #define MARK_PROPERTY_DIRTY(Object, Property) 
@@ -451,4 +487,9 @@ namespace UEPushModelPrivate
 #define COMPARE_ASSIGN_AND_MARK_PROPERTY_DIRTY(ClassName, PropertyName, NewValue, Object) \
 	if (NewValue != PropertyName) { PropertyName = NewValue; }
 
-#endif
\ No newline at end of file
+#if UE_WITH_IRIS
+#define UE_NET_SET_IRIS_MARK_PROPERTY_DIRTY_DELEGATE(...)
+#define UE_NET_SET_IRIS_MARK_PROPERTIES_DIRTY_DELEGATE(...)
+#endif // UE_WITH_IRIS
+
+#endif
diff --git a/Engine/Source/Runtime/Net/Core/Public/Net/Core/PushModel/PushModelMacros.h b/Engine/Source/Runtime/Net/Core/Public/Net/Core/PushModel/PushModelMacros.h
--- a/Engine/Source/Runtime/Net/Core/Public/Net/Core/PushModel/PushModelMacros.h
+++ b/Engine/Source/Runtime/Net/Core/Public/Net/Core/PushModel/PushModelMacros.h
@@ -15,11 +15,11 @@
 
 #define REPLICATED_BASE_CLASS(ClassName) \
 private: \
-int32 NetPushId_Internal = INDEX_NONE; \
-virtual void SetNetPushIdDynamic(const int32 InNetPushId) override final { NetPushId_Internal = InNetPushId; } \
+uint64 NetPushId_Internal = uint64(int64(INDEX_NONE)); \
+virtual void SetNetPushIdDynamic(const uint64 InNetPushId) override final { NetPushId_Internal = InNetPushId; } \
 public: \
-virtual int32 GetNetPushIdDynamic() const override final { return GetNetPushId(); } \
-int32 GetNetPushId() const { return NetPushId_Internal; }
+virtual uint64 GetNetPushIdDynamic() const override final { return GetNetPushId(); } \
+uint64 GetNetPushId() const { return NetPushId_Internal; }
 
 #else
 
diff --git a/Engine/Source/Runtime/Net/Core/Public/Net/Core/Trace/NetDebugName.h b/Engine/Source/Runtime/Net/Core/Public/Net/Core/Trace/NetDebugName.h
--- a/Engine/Source/Runtime/Net/Core/Public/Net/Core/Trace/NetDebugName.h
+++ b/Engine/Source/Runtime/Net/Core/Public/Net/Core/Trace/NetDebugName.h
@@ -3,6 +3,10 @@
 #pragma once
 
 #include "CoreTypes.h"
+#include "Misc/CString.h"
+
+namespace UE::Net
+{
 
 /**
  * NetDebugName, carries both string pointer and a debugname id assigned by the persistent debug storage. DebugNameId is used to be able to avoid hashing already seen NetDebugNames
@@ -15,7 +19,13 @@ struct FNetDebugName
 	mutable FNetDebugNameId DebugNameId = 0;
 };
 
-inline const TCHAR* ToCStr(const FNetDebugName* DebugName)
+
+NETCORE_API const FNetDebugName* CreatePersistentNetDebugName(const TCHAR* Name, uint32 NameLen);
+inline const FNetDebugName* CreatePersistentNetDebugName(const TCHAR* Name) { return CreatePersistentNetDebugName(Name, FCString::Strlen(Name)); }
+
+}
+
+inline const TCHAR* ToCStr(const UE::Net::FNetDebugName* DebugName)
 {
 	if (DebugName && DebugName->Name)
 	{
@@ -27,3 +37,4 @@ inline const TCHAR* ToCStr(const FNetDebugName* DebugName)
 	}
 }
 
+
diff --git a/Engine/Source/Runtime/Net/Core/Public/Net/Core/Trace/NetTrace.h b/Engine/Source/Runtime/Net/Core/Public/Net/Core/Trace/NetTrace.h
--- a/Engine/Source/Runtime/Net/Core/Public/Net/Core/Trace/NetTrace.h
+++ b/Engine/Source/Runtime/Net/Core/Public/Net/Core/Trace/NetTrace.h
@@ -73,7 +73,7 @@
 #define UE_NET_TRACE_POP_SEND_BUNCH(Collector) UE_NET_TRACE_INTERNAL_POP_SEND_BUNCH(Collector)
 
 /** Append events from SrcCollector to Collector */
-#define UE_NET_TRACE_EVENTS(DstCollector, SrcCollector) UE_NET_TRACE_INTERNAL_EVENTS(DstCollector, SrcCollector)
+#define UE_NET_TRACE_EVENTS(DstCollector, SrcCollector, Stream) UE_NET_TRACE_INTERNAL_EVENTS(DstCollector, SrcCollector, Stream)
 
 /** Mark the end of a bunch */
 #define UE_NET_TRACE_END_BUNCH(Collector, Bunch, ...) UE_NET_TRACE_INTERNAL_END_BUNCH(Collector, Bunch, __VA_ARGS__)				
diff --git a/Engine/Source/Runtime/Net/Core/Public/Net/Core/Trace/NetTraceConfig.h b/Engine/Source/Runtime/Net/Core/Public/Net/Core/Trace/NetTraceConfig.h
--- a/Engine/Source/Runtime/Net/Core/Public/Net/Core/Trace/NetTraceConfig.h
+++ b/Engine/Source/Runtime/Net/Core/Public/Net/Core/Trace/NetTraceConfig.h
@@ -8,7 +8,7 @@
 class FNetTraceCollector;
 
 #if !defined(UE_NET_TRACE_ENABLED)
-#	if !(UE_BUILD_SHIPPING || UE_BUILD_TEST) && UE_TRACE_ENABLED
+#	if !(UE_BUILD_SHIPPING) && UE_TRACE_ENABLED
 #		define UE_NET_TRACE_ENABLED 1
 #	else
 #		define UE_NET_TRACE_ENABLED 0
@@ -28,8 +28,12 @@ namespace ENetTraceVerbosity
 
 #if UE_NET_TRACE_ENABLED
 #	ifndef UE_NET_TRACE_COMPILETIME_VERBOSITY
-#		define UE_NET_TRACE_COMPILETIME_VERBOSITY ENetTraceVerbosity::Verbose
+#		if !(UE_BUILD_SHIPPING || UE_BUILD_TEST)
+#			define UE_NET_TRACE_COMPILETIME_VERBOSITY ENetTraceVerbosity::VeryVerbose
+#		else
+#			define UE_NET_TRACE_COMPILETIME_VERBOSITY ENetTraceVerbosity::Verbose
+#		endif
+#	else
+#		define UE_NET_TRACE_COMPILETIME_VERBOSITY ENetTraceVerbosity::None
 #	endif
-#else
-#	define UE_NET_TRACE_COMPILETIME_VERBOSITY ENetTraceVerbosity::None
 #endif
\ No newline at end of file
diff --git a/Engine/Source/Runtime/Net/Core/Public/Net/Core/Trace/Private/NetTraceInternal.h b/Engine/Source/Runtime/Net/Core/Public/Net/Core/Trace/Private/NetTraceInternal.h
--- a/Engine/Source/Runtime/Net/Core/Public/Net/Core/Trace/Private/NetTraceInternal.h
+++ b/Engine/Source/Runtime/Net/Core/Public/Net/Core/Trace/Private/NetTraceInternal.h
@@ -15,6 +15,19 @@
 // $TODO: Add disconnect reason to connection
 // $TODO: Add reporting of final packet sizes after packethandlers
 
+struct FBitReader;
+struct FBitWriter;
+namespace UE::Net
+{
+	class FNetBitStreamReader;
+	class FNetBitStreamWriter;
+}
+
+extern uint32 GetBitStreamPositionForNetTrace(const UE::Net::FNetBitStreamReader& Stream);
+extern uint32 GetBitStreamPositionForNetTrace(const UE::Net::FNetBitStreamWriter& Stream);
+extern uint32 GetBitStreamPositionForNetTrace(const FBitReader& Stream);
+extern uint32 GetBitStreamPositionForNetTrace(const FBitWriter& Stream);
+
 extern NETCORE_API uint32 GNetTraceRuntimeVerbosity;
 
 enum class ENetTracePacketContentEventType : uint8
@@ -71,7 +84,7 @@ struct FNetTracePacketContentEvent
 {
 	union
 	{
-		FNetDebugNameId DebugNameId;
+		UE::Net::FNetDebugNameId DebugNameId;
 		uint32 ObjectId;
 		FNetTraceBunchInfo BunchInfo;
 	};
@@ -106,7 +119,7 @@ public:
 
 	~FNetTraceEventScope();
 
-	void SetEventName(FNetDebugNameId NetTraceNameId);
+	void SetEventName(UE::Net::FNetDebugNameId NetTraceNameId);
 	void SetEventObjectId(uint32 ObjectId);
 	void ExitScope();
 	bool IsValid() const;
@@ -130,7 +143,7 @@ public:
 	FNetTraceNullEventScope(SinkT, T& InStream, FNetTraceCollector* InCollector, uint32 Verbosity)  {};
 	~FNetTraceNullEventScope() {};
 
-	void SetEventName(FNetDebugNameId NetTraceNameId) {};
+	void SetEventName(UE::Net::FNetDebugNameId NetTraceNameId) {};
 	void SetEventObjectId(uint32 ObjectId) {};
 	void ExitScope() {};
 	bool IsValid() const { return false; }
@@ -165,7 +178,7 @@ struct FNetTrace
 	NETCORE_API static void DestroyTraceCollector(FNetTraceCollector* Collector);
 
 	/** Folds events into DstCollector */
-	NETCORE_API static void FoldTraceCollector(FNetTraceCollector* DstCollector, const FNetTraceCollector* SrcCollector, bool IsBunch);
+	NETCORE_API static void FoldTraceCollector(FNetTraceCollector* DstCollector, const FNetTraceCollector* SrcCollector, uint32 Offset);
 	
 	/** Folds bunch into DstCollector, if the src and dst are the same bunch it is assumed that a bunch already is open and that it will be closed */
 	NETCORE_API static void FoldBunchCollector(FNetTraceCollector* DstCollector, const FNetTraceCollector* SrcCollector);
@@ -181,7 +194,7 @@ struct FNetTrace
 	NETCORE_API static void EndPacketContentEvent(FNetTraceCollector& Collector, uint32 EventIndex, uint32 Pos);
 
 	/** Trace PacketContent event */
-	NETCORE_API static void TracePacketContentEvent(FNetTraceCollector& Collector, FNetDebugNameId InNetTraceNameId, uint32 StartPos, uint32 EndPos, uint32 Verbosity);
+	NETCORE_API static void TracePacketContentEvent(FNetTraceCollector& Collector, UE::Net::FNetDebugNameId InNetTraceNameId, uint32 StartPos, uint32 EndPos, uint32 Verbosity);
 
 	/** Mark the beginning of a new Bunch */
 	NETCORE_API static void BeginBunch(FNetTraceCollector& Collector);
@@ -190,7 +203,7 @@ struct FNetTrace
 	NETCORE_API static void DiscardBunch(FNetTraceCollector& Collector);
 
 	/** Trace end of a bunch, all events recorded between BeginBunch/EndBunch will be part of the bunch */
-	NETCORE_API static void EndBunch(FNetTraceCollector& DstCollector, FNetDebugNameId BunchName, uint32 StartPos, uint32 HeaderBits, uint32 BunchBits, const FNetTraceBunchInfo& BunchInfo);
+	NETCORE_API static void EndBunch(FNetTraceCollector& DstCollector, UE::Net::FNetDebugNameId BunchName, uint32 StartPos, uint32 HeaderBits, uint32 BunchBits, const FNetTraceBunchInfo& BunchInfo);
 
 	/** TraceBunch, Commits all events in BunchCollector to DstCollector, if DstCollector == BunchCollector it is assumed that we are Ending the bunch */
 	NETCORE_API static void TraceBunch(FNetTraceCollector& DstCollector, const FNetTraceBunchInfo& BunchInfo, FName BunchName, uint32 StartPos, uint32 HeaderBits, uint32 BunchBits, const FNetTraceCollector* BunchCollector);
@@ -211,7 +224,7 @@ struct FNetTrace
  	NETCORE_API static void TracePacket(uint32 GameInstanceId, uint32 ConnectionId, uint32 PacketSequenceNumber, uint32 PacketBits, ENetTracePacketType PacketType);
 
 	/** Trace that we have created a new object with the given ObjectId */
-	NETCORE_API static void TraceObjectCreated(uint32 GameInstanceId, uint32 NetObjectId, const FNetDebugName* Name, uint64 TypeIdentifier, uint32 OwnerId);
+	NETCORE_API static void TraceObjectCreated(uint32 GameInstanceId, uint32 NetObjectId, const UE::Net::FNetDebugName* Name, uint64 TypeIdentifier, uint32 OwnerId);
 
 	/** Trace that we have created a new object with the given ObjectId */
 	NETCORE_API static void TraceObjectCreated(uint32 GameInstanceId, uint32 NetObjectId, const FName Name, uint64 TypeIdentifier, uint32 OwnerId);
@@ -232,20 +245,20 @@ struct FNetTrace
 	NETCORE_API static void TraceConnectionClosed(uint32 GameInstanceId, uint32 ConnectionId);
 
 	/** Trace stats counters for networking */
-	NETCORE_API static void TracePacketStatsCounter(uint32 GameInstanceId, uint32 ConnectionId, FNetDebugNameId CounterNameId, uint32 StatValue);
-	NETCORE_API static void TraceFrameStatsCounter(uint32 GameInstanceId, FNetDebugNameId CounterNameId,  uint32 StatValue);
+	NETCORE_API static void TracePacketStatsCounter(uint32 GameInstanceId, uint32 ConnectionId, UE::Net::FNetDebugNameId CounterNameId, uint32 StatValue);
+	NETCORE_API static void TraceFrameStatsCounter(uint32 GameInstanceId, UE::Net::FNetDebugNameId CounterNameId,  uint32 StatValue);
 
 	/** Trace the name */
-	NETCORE_API static FNetDebugNameId TraceName(const TCHAR* Name);
+	NETCORE_API static UE::Net::FNetDebugNameId TraceName(const TCHAR* Name);
 
 	/** Trace the FName */
-	NETCORE_API static FNetDebugNameId TraceName(FName Name);
+	NETCORE_API static UE::Net::FNetDebugNameId TraceName(FName Name);
 
 	/** Trace the FNetDebugName */
-	NETCORE_API static FNetDebugNameId TraceName(const FNetDebugName* Name);
+	NETCORE_API static UE::Net::FNetDebugNameId TraceName(const UE::Net::FNetDebugName* Name);
 
 	/** Trace */
-	inline static FNetDebugNameId TraceName(const FString& Name) { return TraceName(ToCStr(Name)); };
+	inline static UE::Net::FNetDebugNameId TraceName(const FString& Name) { return TraceName(ToCStr(Name)); };
 
 	static bool constexpr GetNetTraceVerbosityEnabled(const uint32 V) { return (uint32)UE_NET_TRACE_COMPILETIME_VERBOSITY >= V; }
 
@@ -297,7 +310,7 @@ class FNetTraceCollector
 	UE_NONCOPYABLE(FNetTraceCollector)
 public:
 	// constants
-	enum { MaxNestingLevel = 16 };
+	enum { MaxNestingLevel = 32 };
 	enum { InlinedEventCount = 256 };
 
 	explicit FNetTraceCollector(uint32 InitialEventCount) { Events.SetNumUninitialized(InitialEventCount); Reset(); }
@@ -314,9 +327,9 @@ private:
 	// Bookkeeping
 	uint32 NestingStack[MaxNestingLevel];
 	uint32 OffsetStack[MaxNestingLevel];
-	uint32 EventCount : 24;
-	uint32 CurrentNestingLevel : 4;
-	uint32 OffsetStackLevel : 4;	
+	uint32 EventCount;
+	uint32 CurrentNestingLevel : 16;
+	uint32 OffsetStackLevel : 16;	
 
 	// Bunch tracking
 	uint32 LastBunchEventIndex;
@@ -433,7 +446,7 @@ void FNetTraceEventScope<T>::ExitScope()
 }
 
 template <typename T>
-void FNetTraceEventScope<T>::SetEventName(FNetDebugNameId NetTraceNameId)
+void FNetTraceEventScope<T>::SetEventName(UE::Net::FNetDebugNameId NetTraceNameId)
 {
 	if (EventIndex != FNetTrace::InvalidEventIndex)
 	{
@@ -518,7 +531,7 @@ void FNetTraceCollector::Reset()
 #define UE_NET_TRACE_INTERNAL_BEGIN_BUNCH(Collector) UE_NET_TRACE_DO_IF(Collector, FNetTrace::BeginBunch(*Collector))
 #define UE_NET_TRACE_INTERNAL_DISCARD_BUNCH(Collector) UE_NET_TRACE_DO_IF(Collector, FNetTrace::DiscardBunch(*Collector))
 #define UE_NET_TRACE_INTERNAL_POP_SEND_BUNCH(Collector) UE_NET_TRACE_DO_IF(Collector, FNetTrace::PopSendBunch(*Collector))
-#define UE_NET_TRACE_INTERNAL_EVENTS(Collector, SrcCollector) UE_NET_TRACE_DO_IF(Collector, FNetTrace::FoldTraceCollector(Collector, SrcCollector, false))
+#define UE_NET_TRACE_INTERNAL_EVENTS(Collector, SrcCollector, Stream) UE_NET_TRACE_DO_IF(Collector, FNetTrace::FoldTraceCollector(Collector, SrcCollector, GetBitStreamPositionForNetTrace(Stream)))
 #define UE_NET_TRACE_INTERNAL_END_BUNCH(Collector, Bunch, ...) UE_NET_TRACE_DO_IF(Collector, FNetTrace::TraceBunch(*Collector, MakeBunchInfo(Bunch), __VA_ARGS__))
 #define UE_NET_TRACE_INTERNAL_BUNCH_SCOPE(Collector, Bunch, ...) FNetTraceBunchScope<decltype(Bunch)> PREPROCESSOR_JOIN(NetTraceBunchScope, __LINE__)(Bunch, __VA_ARGS__, Collector)
 
@@ -557,8 +570,8 @@ void FNetTraceCollector::Reset()
 
 #define UE_NET_TRACE_INTERNAL_DYNAMIC_NAME(Name, Collector, StartPos, EndPos, Verbosity) UE_NET_TRACE_DO_IF(FNetTrace::GetCollectorAtVerbosity(Collector, Verbosity), FNetTrace::TracePacketContentEvent(*Collector, FNetTrace::TraceName(Name), StartPos, EndPos, Verbosity))
 #define UE_NET_TRACE_INTERNAL_ASSIGNED_GUID(GameInstanceId, NetGUID, PathName, OwnerId) UE_NET_TRACE_DO_IF(GNetTraceRuntimeVerbosity, FNetTrace::TraceObjectCreated(GameInstanceId, NetGUID.Value, PathName, 0U, OwnerId))
-#define UE_NET_TRACE_INTERNAL_NETHANDLE_CREATED(Handle, DebugName, ProtocolId, OwnerId) FNetTrace::TraceObjectCreated(Handle.GetReplicationSystemId(), Handle.GetIndex(), DebugName, ProtocolId, OwnerId)
-#define UE_NET_TRACE_INTERNAL_NETHANDLE_DESTROYED(Handle) FNetTrace::TraceObjectDestroyed(Handle.GetReplicationSystemId(), Handle.GetIndex())
+#define UE_NET_TRACE_INTERNAL_NETHANDLE_CREATED(Handle, DebugName, ProtocolId, OwnerId) FNetTrace::TraceObjectCreated(Handle.GetReplicationSystemId(), Handle.GetId(), DebugName, ProtocolId, OwnerId)
+#define UE_NET_TRACE_INTERNAL_NETHANDLE_DESTROYED(Handle) FNetTrace::TraceObjectDestroyed(Handle.GetReplicationSystemId(), Handle.GetId())
 #define UE_NET_TRACE_INTERNAL_CONNECTION_CREATED(...) FNetTrace::TraceConnectionCreated(__VA_ARGS__)
 #define UE_NET_TRACE_INTERNAL_CONNECTION_STATE_UPDATED(...) FNetTrace::TraceConnectionStateUpdated(__VA_ARGS__)
 #define UE_NET_TRACE_INTERNAL_CONNECTION_UPDATED(...) FNetTrace::TraceConnectionUpdated(__VA_ARGS__)
