Fixed -NoPCH -DisableUnity compile issue

diff --git a/Engine/Source/Runtime/Net/Core/Private/Net/Core/Trace/PersistentNetDebugNames.cpp b/Engine/Source/Runtime/Net/Core/Private/Net/Core/Trace/PersistentNetDebugNames.cpp
--- a/Engine/Source/Runtime/Net/Core/Private/Net/Core/Trace/PersistentNetDebugNames.cpp
+++ b/Engine/Source/Runtime/Net/Core/Private/Net/Core/Trace/PersistentNetDebugNames.cpp
@@ -6,6 +6,7 @@
 #include "HAL/CriticalSection.h"
 #include "Misc/CoreDelegates.h"
 #include "Misc/MemStack.h"
+#include "Misc/ScopeLock.h"
 
 namespace UE::Net
 {
