Merging CL#22471087

PushModel
* Fix crash that occurs when setting a property dirty inside BeginDestroy()

FObjectKey
* Added ResolveObjPtrEvenIfUnreachable function to access UObject ptrs flagged for destruction.

diff --git a/Engine/Source/Runtime/CoreUObject/Public/UObject/ObjectKey.h b/Engine/Source/Runtime/CoreUObject/Public/UObject/ObjectKey.h
--- a/Engine/Source/Runtime/CoreUObject/Public/UObject/ObjectKey.h
+++ b/Engine/Source/Runtime/CoreUObject/Public/UObject/ObjectKey.h
@@ -108,6 +108,19 @@ public:
 		return WeakPtr.Get(bEvenIfPendingKill);
 	}
 
+	/**
+	 * Attempt to access the object from which this key was constructed, even if it is RF_PendingKill or RF_Unreachable
+	 * @return The object used to construct this key, or nullptr if it is no longer valid
+	 */
+	UObject* ResolveObjectPtrEvenIfUnreachable() const
+	{
+		FWeakObjectPtr WeakPtr;
+		WeakPtr.ObjectIndex = ObjectIndex;
+		WeakPtr.ObjectSerialNumber = ObjectSerialNumber;
+
+		return WeakPtr.GetEvenIfUnreachable();
+	}
+
 	/** Hash function */
 	friend uint32 GetTypeHash(const FObjectKey& Key)
 	{
diff --git a/Engine/Source/Runtime/Engine/Private/NetDriver.cpp b/Engine/Source/Runtime/Engine/Private/NetDriver.cpp
--- a/Engine/Source/Runtime/Engine/Private/NetDriver.cpp
+++ b/Engine/Source/Runtime/Engine/Private/NetDriver.cpp
@@ -1063,6 +1063,8 @@ void UNetDriver::TickFlush(float DeltaSeconds)
 
 	if ( CurrentRealtimeSeconds - LastCleanupTime > CleanupTimeSeconds )
 	{
+		QUICK_SCOPE_CYCLE_COUNTER(STAT_NetUpdateCleanupTime);
+
 		for ( auto It = ReplicationChangeListMap.CreateIterator(); It; ++It )
 		{
 			if ( !It.Value().IsObjectValid() )
diff --git a/Engine/Source/Runtime/Net/Core/Private/Net/Core/PushModel/PushModel.cpp b/Engine/Source/Runtime/Net/Core/Private/Net/Core/PushModel/PushModel.cpp
--- a/Engine/Source/Runtime/Net/Core/Private/Net/Core/PushModel/PushModel.cpp
+++ b/Engine/Source/Runtime/Net/Core/Private/Net/Core/PushModel/PushModel.cpp
@@ -19,12 +19,10 @@ DEFINE_LOG_CATEGORY_STATIC(LogPushModel, All, All);
 
 struct FNetObjectManagerPushIdHelper
 {
-	static void ResetNetPushId(const FObjectKey& InObjectKey)
+	static void ResetObjectNetPushId(UObject* Object)
 	{
-		if (UObject* Object = InObjectKey.ResolveObjectPtr())
-		{
-			FObjectNetPushIdHelper::SetNetPushIdDynamic(Object, UEPushModelPrivate::FNetPushObjectId().GetValue());
-		}
+		check(Object);
+		FObjectNetPushIdHelper::SetNetPushIdDynamic(Object, UEPushModelPrivate::FNetPushObjectId().GetValue());
 	}
 };
 
@@ -272,28 +270,20 @@ namespace UEPushModelPrivate
 			// But we can shrink it.
 
 			// Go ahead and remove any PerObjectStates that aren't being tracked by any NetDrivers.
-			// We have to wait until GC for this, because the NetDrivers will periodically remove
-			// Network Objects that are still alive (but marked Pending Kill) but we don't have a way
-			// to safely clear the Push Model Handles from those objects.
-			//
-			// That means if we tried to remove these items from Push Model tracking, we could end up
-			// with cases where we reassign the Push Model ID to a new object, and the old object could
-			// inadvertently dirty its state.
-			//
-			// In theory, this should never happen because once the object is marked Pending Kill none
-			// of its properties should change again, but it's also possible that calls like BeginDestroy
-			// could modify properties, etc.
-			//
-			// Currently, none of these objects are actually removed though unless the networking system
-			// detects they are PendingKill (their WeakObjectPtr can't be resolved anymore), so there shouldn't
-			// be any cases where we remove these for "still alive" objects.
+			// We have to wait until GC for this, because the NetDriver's TickFlush() will periodically remove
+			// Network Objects that are flagged for destruction.
 			for (auto It = PerObjectStates.CreateIterator(); It; ++It)
 			{
 				if (!It->HasAnyNetDriverStates())
 				{
 					const FObjectKey& ObjectKey = It->GetObjectKey();
 
-					FNetObjectManagerPushIdHelper::ResetNetPushId(ObjectKey);
+					if (UObject* DestroyedObj = ObjectKey.ResolveObjectPtrEvenIfUnreachable())
+					{
+						// If the UObject ptr still exists even if flagged for destruction, reset it's PushModelId to prevent it from trying to set dirty states while it is destroying itself.
+						FNetObjectManagerPushIdHelper::ResetObjectNetPushId(DestroyedObj);
+					}
+
 					ObjectKeyToInternalId.Remove(ObjectKey);
 					It.RemoveCurrent();
 				}
