Added new parameter OldArraySize to FastArraySerializer::FPostReplicatedReceiveParameters containing the size of the fastarray prior the the last receive operation.

diff --git a/Engine/Source/Runtime/Net/Core/Classes/Net/Serialization/FastArraySerializer.h b/Engine/Source/Runtime/Net/Core/Classes/Net/Serialization/FastArraySerializer.h
--- a/Engine/Source/Runtime/Net/Core/Classes/Net/Serialization/FastArraySerializer.h
+++ b/Engine/Source/Runtime/Net/Core/Classes/Net/Serialization/FastArraySerializer.h
@@ -520,6 +520,8 @@ struct FFastArraySerializer
 	 */
 	struct FPostReplicatedReceiveParameters
 	{
+		// This is the size that the array had before the receive.
+		int32 OldArraySize;
 		uint32 bHasMoreUnmappedReferences : 1U;
 	};
 
@@ -695,15 +697,16 @@ private:
 	
 		/** Conditionally invoke PostReplicatedReceive method depending on if is defined or not */
 		template<typename FastArrayType = SerializerType>
-		inline typename TEnableIf<TModels<CPostReplicatedReceiveFuncable, FastArrayType, const FFastArraySerializer::FPostReplicatedReceiveParameters>::Value, void>::Type CallPostReplicatedReceiveOrNot()
+		inline typename TEnableIf<TModels<CPostReplicatedReceiveFuncable, FastArrayType, const FFastArraySerializer::FPostReplicatedReceiveParameters>::Value, void>::Type CallPostReplicatedReceiveOrNot(int32 OldArraySize)
 		{
 			FFastArraySerializer::FPostReplicatedReceiveParameters PostReceivedParameters;
+			PostReceivedParameters.OldArraySize = OldArraySize;
 			PostReceivedParameters.bHasMoreUnmappedReferences = Parms.bOutHasMoreUnmapped;
 			ArraySerializer.PostReplicatedReceive(PostReceivedParameters);
 		}
 
 		template<typename FastArrayType = SerializerType>
-		inline typename TEnableIf<!TModels<CPostReplicatedReceiveFuncable, FastArrayType, const FFastArraySerializer::FPostReplicatedReceiveParameters>::Value, void>::Type CallPostReplicatedReceiveOrNot() {}
+		inline typename TEnableIf<!TModels<CPostReplicatedReceiveFuncable, FastArrayType, const FFastArraySerializer::FPostReplicatedReceiveParameters>::Value, void>::Type CallPostReplicatedReceiveOrNot(int32) {}
 
 		// Validate that deduced FastArrayItemType is valid and that it is the same as the specified one		
 		static_assert(TIsSame<typename TFastArrayTypeHelper<SerializerType>::FastArrayItemType, Type>::Value, "Auto deduced FastArrayItemType is invalid or differs from the specified type. Make sure that the FastArraySerializer has a single replicated array property.");
@@ -1347,7 +1350,7 @@ bool FFastArraySerializer::FastArrayDeltaSerialize(TArray<Type> &Items, FNetDelt
 		if (Parms.bOutSomeObjectsWereMapped)
 		{
 			ArraySerializer.PostReplicatedChange(ChangedIndices, Items.Num());
-			Helper.CallPostReplicatedReceiveOrNot();
+			Helper.CallPostReplicatedReceiveOrNot(Items.Num());
 		}
 		return true;
 	}
@@ -1464,6 +1467,7 @@ bool FFastArraySerializer::FastArrayDeltaSerialize(TArray<Type> &Items, FNetDelt
 			return false;
 		}
 
+		const int32 OldNumItems = Items.Num();
 		TArray<int32, TInlineAllocator<8>> ChangedIndices;
 		TArray<int32, TInlineAllocator<8>> AddedIndices;
 
@@ -1574,7 +1578,7 @@ bool FFastArraySerializer::FastArrayDeltaSerialize(TArray<Type> &Items, FNetDelt
 
 		Helper.template PostReceiveCleanup<>(Header, ChangedIndices, AddedIndices, ArraySerializer.GuidReferencesMap);
 
-		Helper.CallPostReplicatedReceiveOrNot();
+		Helper.CallPostReplicatedReceiveOrNot(OldNumItems);
 	}
 
 	return true;
@@ -1686,7 +1690,7 @@ bool FFastArraySerializer::FastArrayDeltaSerialize_DeltaSerializeStructs(TArray<
 		if (Parms.bOutSomeObjectsWereMapped)
 		{
 			ArraySerializer.PostReplicatedChange(ChangedIndices, Items.Num());
-			Helper.CallPostReplicatedReceiveOrNot();
+			Helper.CallPostReplicatedReceiveOrNot(Items.Num());
 		}
 
 		return true;
@@ -1789,6 +1793,7 @@ bool FFastArraySerializer::FastArrayDeltaSerialize_DeltaSerializeStructs(TArray<
 			return false;
 		}
 
+		const int32 OldNumItems = Items.Num();
 		TArray<int32, TInlineAllocator<8>> ChangedIndices;
 		TArray<int32, TInlineAllocator<8>> AddedIndices;
 
@@ -1808,7 +1813,7 @@ bool FFastArraySerializer::FastArrayDeltaSerialize_DeltaSerializeStructs(TArray<
 
 		Helper.template PostReceiveCleanup<>(Header, ChangedIndices, AddedIndices, ArraySerializer.GuidReferencesMap_StructDelta);
 
-		Helper.CallPostReplicatedReceiveOrNot();
+		Helper.CallPostReplicatedReceiveOrNot(OldNumItems);
 	}
 
 	return true;
