Iris
FastArraySerialzier
- PostReplicatedReceive is now always called after all calls to FastArrayDeltaSerialize on the receiving side

diff --git a/Engine/Source/Runtime/Net/Core/Classes/Net/Serialization/FastArraySerializer.h b/Engine/Source/Runtime/Net/Core/Classes/Net/Serialization/FastArraySerializer.h
--- a/Engine/Source/Runtime/Net/Core/Classes/Net/Serialization/FastArraySerializer.h
+++ b/Engine/Source/Runtime/Net/Core/Classes/Net/Serialization/FastArraySerializer.h
@@ -1226,6 +1226,8 @@ bool FFastArraySerializer::FastArrayDeltaSerialize(TArray<Type> &Items, FNetDelt
 			}
 		}
 
+		Helper.CallPostReplicatedReceiveOrNot(Items.Num());
+
 		return true;
 	}
 
@@ -1350,8 +1352,9 @@ bool FFastArraySerializer::FastArrayDeltaSerialize(TArray<Type> &Items, FNetDelt
 		if (Parms.bOutSomeObjectsWereMapped)
 		{
 			ArraySerializer.PostReplicatedChange(ChangedIndices, Items.Num());
-			Helper.CallPostReplicatedReceiveOrNot(Items.Num());
 		}
+		Helper.CallPostReplicatedReceiveOrNot(Items.Num());
+
 		return true;
 	}
 
@@ -1670,6 +1673,8 @@ bool FFastArraySerializer::FastArrayDeltaSerialize_DeltaSerializeStructs(TArray<
 	{
 		UE_LOG(LogNetFastTArray, Log, TEXT("FastArrayDeltaSerialize_DeltaSerializeStruct for %s. %s. GatherGuidReferences"), *InnerStruct->GetName(), *InnerStruct->GetOwnerStruct()->GetName());
 		Parms.NetSerializeCB->GatherGuidReferencesForFastArray(DeltaSerializeParams);
+		Helper.CallPostReplicatedReceiveOrNot(Items.Num());
+
 		return true;
 	}
 	else if (Parms.MoveGuidToUnmapped)
@@ -1690,8 +1695,8 @@ bool FFastArraySerializer::FastArrayDeltaSerialize_DeltaSerializeStructs(TArray<
 		if (Parms.bOutSomeObjectsWereMapped)
 		{
 			ArraySerializer.PostReplicatedChange(ChangedIndices, Items.Num());
-			Helper.CallPostReplicatedReceiveOrNot(Items.Num());
 		}
+		Helper.CallPostReplicatedReceiveOrNot(Items.Num());
 
 		return true;
 	}
