Fixed incorrect deprecation version.

#rnx

diff --git a/Engine/Source/Runtime/Net/Core/Public/Net/Core/PropertyConditions/RepChangedPropertyTracker.h b/Engine/Source/Runtime/Net/Core/Public/Net/Core/PropertyConditions/RepChangedPropertyTracker.h
--- a/Engine/Source/Runtime/Net/Core/Public/Net/Core/PropertyConditions/RepChangedPropertyTracker.h
+++ b/Engine/Source/Runtime/Net/Core/Public/Net/Core/PropertyConditions/RepChangedPropertyTracker.h
@@ -43,7 +43,7 @@ public:
 		* @param RepIndex		Replication index for the Property.
 		* @param bIsActive		The new Active state.
 		*/
-	UE_DEPRECATED(5.3, "Please use FNetPropertyConditionManager::SetPropertyActiveOverride instead.")
+	UE_DEPRECATED(5.2, "Please use FNetPropertyConditionManager::SetPropertyActiveOverride instead.")
 	virtual void SetCustomIsActiveOverride(UObject* OwningObject, const uint16 RepIndex, const bool bIsActive) override;
 
 	/**
