Removed incorrect WITH_PUSH_MODEL usage. The functionality is needed regardless of push model.
Fixed typo.

diff --git a/Engine/Source/Runtime/Net/Core/Private/Net/Core/DirtyNetObjectTracker/GlobalDirtyNetObjectTracker.cpp b/Engine/Source/Runtime/Net/Core/Private/Net/Core/DirtyNetObjectTracker/GlobalDirtyNetObjectTracker.cpp
--- a/Engine/Source/Runtime/Net/Core/Private/Net/Core/DirtyNetObjectTracker/GlobalDirtyNetObjectTracker.cpp
+++ b/Engine/Source/Runtime/Net/Core/Private/Net/Core/DirtyNetObjectTracker/GlobalDirtyNetObjectTracker.cpp
@@ -20,33 +20,26 @@ private:
 
 	static TSet<FNetHandle> EmptyDirtyObjects;
 
-#if WITH_PUSH_MODEL
 	TSet<FNetHandle> DirtyObjects;
 	FNetBitArray AssignedHandleIndices;
 	FNetBitArray Pollers;
 	uint32 PollerCount = 0;
-#endif
 };
 
 TSet<FNetHandle> FGlobalDirtyNetObjectTracker::FPimpl::EmptyDirtyObjects;
 
-#if WITH_PUSH_MODEL
 FGlobalDirtyNetObjectTracker::FPimpl* FGlobalDirtyNetObjectTracker::Instance = nullptr;
-#endif
 
 void FGlobalDirtyNetObjectTracker::MarkNetObjectStateDirty(FNetHandle NetHandle)
 {
-#if WITH_PUSH_MODEL
 	if (Instance && Instance->PollerCount > 0)
 	{
 		Instance->DirtyObjects.Add(NetHandle);
 	}
-#endif
 }
 
 FGlobalDirtyNetObjectTracker::FPollHandle FGlobalDirtyNetObjectTracker::CreatePoller()
 {
-#if WITH_PUSH_MODEL
 	if (Instance)
 	{
 		if (Instance->PollerCount >= Instance->AssignedHandleIndices.GetNumBits())
@@ -65,14 +58,12 @@ FGlobalDirtyNetObjectTracker::FPollHandle FGlobalDirtyNetObjectTracker::CreatePo
 		++Instance->PollerCount;
 		return FPollHandle(HandleIndex);
 	}
-#endif
 
 	return FPollHandle();
 }
 
 void FGlobalDirtyNetObjectTracker::DestroyPoller(uint32 HandleIndex)
 {
-#if WITH_PUSH_MODEL
 	if (HandleIndex == FPollHandle::InvalidIndex)
 	{
 		return;
@@ -92,25 +83,21 @@ void FGlobalDirtyNetObjectTracker::DestroyPoller(uint32 HandleIndex)
 			Instance->DirtyObjects.Reset();
 		}
 	}
-#endif
 }
 
 const TSet<FNetHandle>& FGlobalDirtyNetObjectTracker::GetDirtyNetObjects(const FPollHandle& Handle)
 {
-#if WITH_PUSH_MODEL
 	if (Instance && Handle.IsValid())
 	{
 		Instance->Pollers.SetBit(Handle.Index);
 		return Instance->DirtyObjects;
 	}
-#endif
 
 	return FPimpl::EmptyDirtyObjects;
 }
 
 void FGlobalDirtyNetObjectTracker::ResetDirtyNetObjects(const FPollHandle& Handle)
 {
-#if WITH_PUSH_MODEL
 	if (Instance && Handle.IsValid())
 	{
 		Instance->Pollers.ClearBit(Handle.Index);
@@ -119,23 +106,18 @@ void FGlobalDirtyNetObjectTracker::ResetDirtyNetObjects(const FPollHandle& Handl
 			Instance->DirtyObjects.Reset();
 		}
 	}
-#endif
 }
 
 void FGlobalDirtyNetObjectTracker::Init()
 {
-#if WITH_PUSH_MODEL
-	checkf(Instance == nullptr, TEXT("%s"), TEXT("Only one FNetHandleManager instance may exist."));
+	checkf(Instance == nullptr, TEXT("%s"), TEXT("Only one FGlobalDirtyNetObjectTracker instance may exist."));
 	Instance = new FGlobalDirtyNetObjectTracker::FPimpl();
-#endif
 }
 
 void FGlobalDirtyNetObjectTracker::Deinit()
 {
-#if WITH_PUSH_MODEL
 	delete Instance;
 	Instance = nullptr;
-#endif
 }
 
 FGlobalDirtyNetObjectTracker::FPimpl::FPimpl()
@@ -154,7 +136,6 @@ FGlobalDirtyNetObjectTracker::FPimpl::~FPimpl()
 
 void FGlobalDirtyNetObjectTracker::FPimpl::ResetDirtyNetObjects()
 {
-#if WITH_PUSH_MODEL
 	if (!ensureMsgf(Pollers.IsNoBitSet(), TEXT("FGlobalDirtyNetObjectTracker poller %u forgot to call ResetDirtNetObjects."), Pollers.FindFirstOne()))
 	{
 		Pollers.Reset();
@@ -162,7 +143,6 @@ void FGlobalDirtyNetObjectTracker::FPimpl::ResetDirtyNetObjects()
 		// DirtyObjects should already be reset if there are no pollers.
 		DirtyObjects.Reset();
 	}
-#endif
 }
 
 }
diff --git a/Engine/Source/Runtime/Net/Core/Public/Net/Core/DirtyNetObjectTracker/GlobalDirtyNetObjectTracker.h b/Engine/Source/Runtime/Net/Core/Public/Net/Core/DirtyNetObjectTracker/GlobalDirtyNetObjectTracker.h
--- a/Engine/Source/Runtime/Net/Core/Public/Net/Core/DirtyNetObjectTracker/GlobalDirtyNetObjectTracker.h
+++ b/Engine/Source/Runtime/Net/Core/Public/Net/Core/DirtyNetObjectTracker/GlobalDirtyNetObjectTracker.h
@@ -81,9 +81,7 @@ private:
 	NETCORE_API static void DestroyPoller(uint32 HandleIndex);
 
 	class FPimpl;
-#if WITH_PUSH_MODEL
 	static FPimpl* Instance;
-#endif
 };
 
 inline void MarkNetObjectStateDirty(FNetHandle NetHandle)
