Fix or work around static analysis warnings/false positives in clang 15.

diff --git a/Engine/Source/Runtime/Net/Core/Private/Net/Core/Misc/PacketAudit.cpp b/Engine/Source/Runtime/Net/Core/Private/Net/Core/Misc/PacketAudit.cpp
--- a/Engine/Source/Runtime/Net/Core/Private/Net/Core/Misc/PacketAudit.cpp
+++ b/Engine/Source/Runtime/Net/Core/Private/Net/Core/Misc/PacketAudit.cpp
@@ -254,6 +254,7 @@ FPacketAudit::FPacketAudit(FPlatformProcess::FSemaphore* InGameMutex)
 	if (SendSharedMemory == nullptr || ReceiveSharedMemory == nullptr)
 	{
 		LowLevelFatalError(TEXT("You need to run UnrealEngine as administrator, for packet auditing to work."));
+		return;
 	}
 
 	FSharedMemoryWriter WipeSendSharedMemory((uint8*)SendSharedMemory->GetAddress());
@@ -305,6 +306,14 @@ void FPacketAudit::Init()
 #endif
 }
 
+void FPacketAudit::Destruct()
+{
+#if !UE_BUILD_SHIPPING
+	delete GPacketAuditor;
+	GPacketAuditor = nullptr;
+#endif
+}
+
 void FPacketAudit::AddStage_Internal(FString StageName, FBitWriter& OutPacket, bool bByteAligned)
 {
 #if !UE_BUILD_SHIPPING
diff --git a/Engine/Source/Runtime/Net/Core/Public/Net/Core/Misc/PacketAudit.h b/Engine/Source/Runtime/Net/Core/Public/Net/Core/Misc/PacketAudit.h
--- a/Engine/Source/Runtime/Net/Core/Public/Net/Core/Misc/PacketAudit.h
+++ b/Engine/Source/Runtime/Net/Core/Public/Net/Core/Misc/PacketAudit.h
@@ -115,16 +115,7 @@ public:
 	/**
 	 * Destroy GPacketAuditor
 	 */
-	static void Destruct()
-	{
-#if !UE_BUILD_SHIPPING
-		if (GPacketAuditor != nullptr)
-		{
-			delete GPacketAuditor;
-			GPacketAuditor = nullptr;
-		}
-#endif
-	}
+	static void Destruct();
 
 	/**
 	 * On the send side, marks a named stage during packet writing, which should be audited on the receive side.
diff --git a/Engine/Source/Runtime/PacketHandlers/PacketHandler/Private/PacketHandler.cpp b/Engine/Source/Runtime/PacketHandlers/PacketHandler/Private/PacketHandler.cpp
--- a/Engine/Source/Runtime/PacketHandlers/PacketHandler/Private/PacketHandler.cpp
+++ b/Engine/Source/Runtime/PacketHandlers/PacketHandler/Private/PacketHandler.cpp
@@ -76,6 +76,14 @@ static FAutoConsoleVariableRef CVarNetPacketHandlerTimeguardLimit(
   #define NET_LIGHTWEIGHT_TIME_GUARD_END( Name, NameStringCode )
 #endif
 
+/**
+ * BufferedPacket
+ */
+
+BufferedPacket::~BufferedPacket()
+{
+	delete [] Data;
+}
 
 /**
  * PacketHandler
diff --git a/Engine/Source/Runtime/PacketHandlers/PacketHandler/Public/PacketHandler.h b/Engine/Source/Runtime/PacketHandlers/PacketHandler/Public/PacketHandler.h
--- a/Engine/Source/Runtime/PacketHandlers/PacketHandler/Public/PacketHandler.h
+++ b/Engine/Source/Runtime/PacketHandlers/PacketHandler/Public/PacketHandler.h
@@ -211,10 +211,7 @@ public:
 	/**
 	 * Base destructor
 	 */
-	~BufferedPacket()
-	{
-		delete [] Data;
-	}
+	~BufferedPacket();
 
 	void CountBytes(FArchive& Ar) const
 	{
