FastArraySerializer is no longer allowed to modify state of Archetype or CDO when initializing base state as this lead to inconsistencies with assignment of ReplicationIDs
- Fixed issue with missing detection of items being marked dirty

diff --git a/Engine/Source/Runtime/Engine/Private/DataReplication.cpp b/Engine/Source/Runtime/Engine/Private/DataReplication.cpp
--- a/Engine/Source/Runtime/Engine/Private/DataReplication.cpp
+++ b/Engine/Source/Runtime/Engine/Private/DataReplication.cpp
@@ -24,10 +24,6 @@
 #include "Net/Core/PushModel/Types/PushModelPerNetDriverState.h"
 #include "Net/RPCDoSDetection.h"
 
-#if UE_WITH_IRIS
-#include "Iris/IrisConfig.h"
-#endif
-
 DECLARE_LLM_MEMORY_STAT(TEXT("NetObjReplicator"), STAT_NetObjReplicatorLLM, STATGROUP_LLMFULL);
 LLM_DEFINE_TAG(NetObjReplicator, NAME_None, TEXT("Networking"), GET_STATFNAME(STAT_NetObjReplicatorLLM), GET_STATFNAME(STAT_NetworkingSummaryLLM));
 
@@ -377,13 +373,8 @@ bool FObjectReplicator::SendCustomDeltaProperty(UObject* InObject, uint16 Custom
 	Parms.Connection = Connection;
 	Parms.bInternalAck = Connection->IsInternalAck();
 
-#if UE_WITH_IRIS
-	if (UE::Net::ShouldUseIrisReplication())
-	{
-		// When initializing baselines we should not modify the source data if it originates from the CDO or archetype
-		Parms.bIsInitializingBaseFromDefault = Parms.Object && Parms.Object->HasAnyFlags(RF_ClassDefaultObject | RF_ArchetypeObject);
-	}
-#endif
+	// When initializing baselines we should not modify the source data if it originates from the CDO or archetype
+	Parms.bIsInitializingBaseFromDefault = Parms.Object && Parms.Object->HasAnyFlags(RF_ClassDefaultObject | RF_ArchetypeObject);
 
 	return FNetSerializeCB::SendCustomDeltaProperty(*RepLayout, Parms, CustomDeltaIndex);
 }
diff --git a/Engine/Source/Runtime/Net/Core/Classes/Net/Serialization/FastArraySerializer.h b/Engine/Source/Runtime/Net/Core/Classes/Net/Serialization/FastArraySerializer.h
--- a/Engine/Source/Runtime/Net/Core/Classes/Net/Serialization/FastArraySerializer.h
+++ b/Engine/Source/Runtime/Net/Core/Classes/Net/Serialization/FastArraySerializer.h
@@ -280,7 +280,6 @@ public:
 	int32 ArrayReplicationKey;
 };
 
-
 /** Base struct for items using Fast TArray Replication */
 USTRUCT()
 struct FFastArraySerializerItem
@@ -436,7 +435,7 @@ struct FFastArraySerializer
 	// Property index of this array in the owning object's replication layout
 	int32 RepIndex;
 #endif // WITH_PUSH_MODEL
-	
+
 	/** This must be called if you add or change an item in the array */
 	void MarkItemDirty(FFastArraySerializerItem & Item)
 	{
@@ -675,7 +674,7 @@ private:
 		 */
 		void BuildChangedAndDeletedBuffersFromDefault(
 			TMap<int32, int32>& NewIDToKeyMap,
-			TArray<FFastArraySerializer_FastArrayDeltaSerialize_FIdxIDPair, TInlineAllocator<8>>& ChangedElements);
+			TArray<FFastArraySerializer_FastArrayDeltaSerialize_FIdxIDPair, TInlineAllocator<8>>& ChangedElements) const;
 
 		/** Writes out a FFastArraySerializerHeader */
 		void WriteDeltaHeader(FFastArraySerializerHeader& Header) const;
@@ -855,17 +854,12 @@ bool FFastArraySerializer::TFastArraySerializeHelper<Type, SerializerType>::Cond
 template<typename Type, typename SerializerType>
 void FFastArraySerializer::TFastArraySerializeHelper<Type, SerializerType>::BuildChangedAndDeletedBuffersFromDefault(
 	TMap<int32, int32>& NewIDToKeyMap,
-	TArray<FFastArraySerializer_FastArrayDeltaSerialize_FIdxIDPair, TInlineAllocator<8>>& ChangedElements)
+	TArray<FFastArraySerializer_FastArrayDeltaSerialize_FIdxIDPair, TInlineAllocator<8>>& ChangedElements) const
 {
 	const int32 NumConsideredItems = CalcNumItemsForConsideration();
 
 	// Verify assumptions, we never expect the default state to have replicated and thus IDCounter and ArrayReplicationKey should be at the defaults
 	check(Parms.bIsInitializingBaseFromDefault);
-	check(ArraySerializer.IDCounter == 0 && ArraySerializer.ArrayReplicationKey == 0);
-
-	// We fake assignment of ID to avoid modifying the default state
-	int32 FakeIDCounter = ArraySerializer.IDCounter;
-	const int32 FakeReplicationKey = 0;
 
 	//-----------------------------------------------------------------
 	// When initializing from default we assume that all items are new
@@ -881,15 +875,18 @@ void FFastArraySerializer::TFastArraySerializeHelper<Type, SerializerType>::Buil
 			continue;
 		}
 
-		// Item should not have been touched before
-		check(Item.ReplicationID == INDEX_NONE && Item.ReplicationKey == INDEX_NONE);
-
-		const int32 FakeID = FakeIDCounter + i;
+		// When initializing from the CDO or archetype we do not allow the state to be modified as this will potentially lead to mismatch in ID assignment
+		if (Item.ReplicationID == INDEX_NONE)
+		{
+			UE_LOG(LogNetFastTArray, Log, TEXT("    FastArraySerializer::BuildChangedAndDeletedBuffersFromDefault Item with uninitialized ReplicationID detected in. Struct: %s, ItemIndex: %i"), *Parms.DebugName, i);
+			continue;
+		}
 
-		NewIDToKeyMap.Add(FakeID, FakeReplicationKey);
+		NewIDToKeyMap.Add(Item.ReplicationID, Item.ReplicationKey);
 
-		UE_LOG(LogNetFastTArray, Log, TEXT("       New! Element ID: %d. %s"), FakeID, *Item.GetDebugString());
-		ChangedElements.Add(FFastArraySerializer_FastArrayDeltaSerialize_FIdxIDPair(i, FakeID));
+		UE_LOG(LogNetFastTArray, Log, TEXT("       New! Element ID: %d. %s"), Item.ReplicationID, *Item.GetDebugString());
+		// New
+		ChangedElements.Add(FFastArraySerializer_FastArrayDeltaSerialize_FIdxIDPair(i, Item.ReplicationID));
 	}
 }
 
@@ -920,10 +917,14 @@ void FFastArraySerializer::TFastArraySerializeHelper<Type, SerializerType>::Buil
 			// On clients, this will skip items that were added predictively.
 			continue;
 		}
+
+		// The item really should have a valid ReplicationID but in the case of loading from a save game, or initializing from a default state/archetype with existing data
+		// items may not have been marked dirty individually. Its ok to just assign them one here.
 		if (Item.ReplicationID == INDEX_NONE)
 		{
 			ArraySerializer.MarkItemDirty(Item);
 		}
+
 		NewIDToKeyMap.Add(Item.ReplicationID, Item.ReplicationKey);
 
 		const int32* OldValuePtr = OldIDToKeyMap ? OldIDToKeyMap->Find(Item.ReplicationID) : NULL;
@@ -948,8 +949,6 @@ void FFastArraySerializer::TFastArraySerializeHelper<Type, SerializerType>::Buil
 		{
 			UE_LOG(LogNetFastTArray, Log, TEXT("       New! Element ID: %d. %s"), Item.ReplicationID, *Item.GetDebugString());
 
-			// The item really should have a valid ReplicationID but in the case of loading from a save game,
-			// items may not have been marked dirty individually. Its ok to just assign them one here.
 			// New
 			ChangedElements.Add(FFastArraySerializer_FastArrayDeltaSerialize_FIdxIDPair(i, Item.ReplicationID));
 			++DeleteCount; // We added something new, so our initial DeleteCount value must be incremented.
@@ -1401,7 +1400,7 @@ bool FFastArraySerializer::FastArrayDeltaSerialize(TArray<Type> &Items, FNetDelt
 
 		check(Parms.NewState);
 		*Parms.NewState = MakeShareable( NewState );
-		TMap<int32, int32> & NewMap = NewState->IDToCLMap;
+		TMap<int32, int32>& NewMap = NewState->IDToCLMap;
 		NewState->ArrayReplicationKey = ArraySerializer.ArrayReplicationKey;
 
 		FFastArraySerializerHeader Header{
@@ -1420,14 +1419,14 @@ bool FFastArraySerializer::FastArrayDeltaSerialize(TArray<Type> &Items, FNetDelt
 		else
 		{
 			Helper.BuildChangedAndDeletedBuffers(NewMap, OldMap, ChangedElements, Header.DeletedIndices);
+
+			// The array replication key may have changed while adding new elements (in the call to BuildChangedAndDeletedBuffers above)
+			NewState->ArrayReplicationKey = ArraySerializer.ArrayReplicationKey;
 		}
-		
+
 		// Note: we used to early return false here if nothing had changed, but we still need to send
 		// a bunch with the array key / base key, so that clients can look for implicit deletes.
-
-		// The array replication key may have changed while adding new elements (in the call to MarkItemDirty above)
-		NewState->ArrayReplicationKey = ArraySerializer.ArrayReplicationKey;
-
+		
 		//----------------------
 		// Write it out.
 		//----------------------
@@ -1754,14 +1753,14 @@ bool FFastArraySerializer::FastArrayDeltaSerialize_DeltaSerializeStructs(TArray<
 		else
 		{
 			Helper.BuildChangedAndDeletedBuffers(NewItemMap, OldItemMap, ChangedElements, Header.DeletedIndices);
+
+			// The array replication key may have changed while adding new elemnts (in the call to MarkItemDirty above)
+			NewState->ArrayReplicationKey = ArraySerializer.ArrayReplicationKey;
 		}
 
 		// Note: we used to early return false here if nothing had changed, but we still need to send
 		// a bunch with the array key / base key, so that clients can look for implicit deletes.
 
-		// The array replication key may have changed while adding new elemnts (in the call to MarkItemDirty above)
-		NewState->ArrayReplicationKey = ArraySerializer.ArrayReplicationKey;
-
 		//----------------------
 		// Write it out.
 		//----------------------
