Iris - Use own TUniquePtr instead of TlsAutoCleanup. Want to delete our thread buffer at times.

diff --git a/Engine/Source/Runtime/Net/Core/Private/Net/Core/Trace/NetTraceInternal.cpp b/Engine/Source/Runtime/Net/Core/Private/Net/Core/Trace/NetTraceInternal.cpp
--- a/Engine/Source/Runtime/Net/Core/Private/Net/Core/Trace/NetTraceInternal.cpp
+++ b/Engine/Source/Runtime/Net/Core/Private/Net/Core/Trace/NetTraceInternal.cpp
@@ -4,8 +4,6 @@
 #include "Containers/Map.h"
 #include "Net/Core/Trace/NetDebugName.h"
 #include "Net/Core/Trace/Reporters/NetTraceReporter.h"
-#include "HAL/PlatformTLS.h"
-#include "HAL/TlsAutoCleanup.h"
 #include "Hash/CityHash.h"
 #include "Trace/Trace.h"
 #include "UObject/NameTypes.h"
@@ -37,7 +35,7 @@ struct FNetTraceInternal
 		ENetTraceVersion_BunchChannelInfo = 3,
 	};
 
-	struct FThreadBuffer : public FTlsAutoCleanup
+	struct FThreadBuffer
 	{
 		// Map FName to NameId, used to know what FNames we have traced
 		TMap<FName, UE::Net::FNetDebugNameId> DynamicFNameToNameIdMap;
@@ -49,21 +47,16 @@ struct FNetTraceInternal
 	// Get next NameId used to track what we already have traced
 	uint32 static GetNextNameId();
 
-	FORCENOINLINE static FThreadBuffer* CreateThreadBuffer();
+	static inline FThreadBuffer* CreateThreadBuffer();
 
-	static thread_local FThreadBuffer* ThreadBuffer;
+	static inline thread_local TUniquePtr<FThreadBuffer> ThreadBuffer;
 	static constexpr ENetTraceVersion NetTraceVersion = ENetTraceVersion::ENetTraceVersion_BunchChannelInfo;
 };
 
-thread_local FNetTraceInternal::FThreadBuffer* FNetTraceInternal::ThreadBuffer = nullptr;
-
-
 FNetTraceInternal::FThreadBuffer* FNetTraceInternal::CreateThreadBuffer()
 {
-	ThreadBuffer = new FThreadBuffer();
-	ThreadBuffer->Register();
-
-	return ThreadBuffer;
+	ThreadBuffer = MakeUnique<FThreadBuffer>();
+	return ThreadBuffer.Get();
 }
 
 uint32 FNetTraceInternal::GetNextNameId()
@@ -91,8 +84,7 @@ void FNetTrace::SetTraceVerbosity(uint32 Verbosity)
 		{
 			UE::Trace::ToggleChannel(TEXT("NetChannel"), false);
 
-			delete FNetTraceInternal::ThreadBuffer;
-			FNetTraceInternal::ThreadBuffer = nullptr;
+			FNetTraceInternal::ThreadBuffer.Reset();
 		}		
 	}
 
@@ -518,7 +510,7 @@ UE::Net::FNetDebugNameId FNetTrace::TraceName(const TCHAR* Name)
 	}
 
 	// Get Thread buffer
-	FNetTraceInternal::FThreadBuffer* ThreadBuffer = FNetTraceInternal::ThreadBuffer;
+	FNetTraceInternal::FThreadBuffer* ThreadBuffer = FNetTraceInternal::ThreadBuffer.Get();
 	if (!ThreadBuffer)
 	{
 		ThreadBuffer = FNetTraceInternal::CreateThreadBuffer();
@@ -552,7 +544,7 @@ UE::Net::FNetDebugNameId FNetTrace::TraceName(FName Name)
 	}
 
 	// Get Thread buffer
-	FNetTraceInternal::FThreadBuffer* ThreadBuffer = FNetTraceInternal::ThreadBuffer;
+	FNetTraceInternal::FThreadBuffer* ThreadBuffer = FNetTraceInternal::ThreadBuffer.Get();
 	if (!ThreadBuffer)
 	{
 		ThreadBuffer = FNetTraceInternal::CreateThreadBuffer();
