Support uint64 network object ids in traces

diff --git a/Engine/Source/Developer/TraceInsights/Private/Insights/NetworkingProfiler/ViewModels/PacketContentViewDrawHelper.cpp b/Engine/Source/Developer/TraceInsights/Private/Insights/NetworkingProfiler/ViewModels/PacketContentViewDrawHelper.cpp
--- a/Engine/Source/Developer/TraceInsights/Private/Insights/NetworkingProfiler/ViewModels/PacketContentViewDrawHelper.cpp
+++ b/Engine/Source/Developer/TraceInsights/Private/Insights/NetworkingProfiler/ViewModels/PacketContentViewDrawHelper.cpp
@@ -36,7 +36,7 @@ FPacketContentViewDrawStateBuilder::FPacketContentViewDrawStateBuilder(FPacketCo
 
 ////////////////////////////////////////////////////////////////////////////////////////////////////
 
-void FPacketContentViewDrawStateBuilder::AddEvent(const TraceServices::FNetProfilerContentEvent& Event, const TCHAR* EventName, uint32 NetId)
+void FPacketContentViewDrawStateBuilder::AddEvent(const TraceServices::FNetProfilerContentEvent& Event, const TCHAR* EventName, uint64 NetId)
 {
 	DrawState.Events.AddUninitialized();
 	FNetworkPacketEvent& PacketEvent = DrawState.Events.Last();
@@ -233,7 +233,7 @@ void FPacketContentViewDrawStateBuilder::AddEvent(const TraceServices::FNetProfi
 
 			if (Event.ObjectInstanceIndex != 0)
 			{
-				Builder.Appendf(TEXT(" (NetId:%u, "), NetId);
+				Builder.Appendf(TEXT(" (NetId:%" UINT64_FMT ", "), NetId);
 			}
 			else
 			{
diff --git a/Engine/Source/Developer/TraceInsights/Private/Insights/NetworkingProfiler/ViewModels/PacketContentViewDrawHelper.h b/Engine/Source/Developer/TraceInsights/Private/Insights/NetworkingProfiler/ViewModels/PacketContentViewDrawHelper.h
--- a/Engine/Source/Developer/TraceInsights/Private/Insights/NetworkingProfiler/ViewModels/PacketContentViewDrawHelper.h
+++ b/Engine/Source/Developer/TraceInsights/Private/Insights/NetworkingProfiler/ViewModels/PacketContentViewDrawHelper.h
@@ -22,7 +22,7 @@ struct FNetworkPacketEvent
 {
 	uint32 EventTypeIndex;
 	uint32 ObjectInstanceIndex;
-	uint32 NetId;
+	uint64 NetId;
 	uint32 BitOffset;
 	uint32 BitSize;
 	uint32 Level;
@@ -36,7 +36,7 @@ struct FNetworkPacketEvent
 		, Level(0)
 	{}
 
-	FNetworkPacketEvent(uint32 InEventTypeIndex, uint32 InObjectInstanceIndex, uint32 InNetId, uint32 InBitOffset, uint32 InBitSize, uint32 InLevel)
+	FNetworkPacketEvent(uint32 InEventTypeIndex, uint32 InObjectInstanceIndex, uint64 InNetId, uint32 InBitOffset, uint32 InBitSize, uint32 InLevel)
 		: EventTypeIndex(InEventTypeIndex)
 		, ObjectInstanceIndex(InObjectInstanceIndex)
 		, NetId(InNetId)
@@ -167,7 +167,7 @@ public:
 	FPacketContentViewDrawStateBuilder(const FPacketContentViewDrawStateBuilder&) = delete;
 	FPacketContentViewDrawStateBuilder& operator=(const FPacketContentViewDrawStateBuilder&) = delete;
 
-	void AddEvent(const TraceServices::FNetProfilerContentEvent& Event, const TCHAR* Name, uint32 NetId);
+	void AddEvent(const TraceServices::FNetProfilerContentEvent& Event, const TCHAR* Name, uint64 NetId);
 	void Flush();
 
 	int32 GetMaxDepth() const { return MaxDepth; }
diff --git a/Engine/Source/Developer/TraceInsights/Private/Insights/NetworkingProfiler/Widgets/SPacketContentView.cpp b/Engine/Source/Developer/TraceInsights/Private/Insights/NetworkingProfiler/Widgets/SPacketContentView.cpp
--- a/Engine/Source/Developer/TraceInsights/Private/Insights/NetworkingProfiler/Widgets/SPacketContentView.cpp
+++ b/Engine/Source/Developer/TraceInsights/Private/Insights/NetworkingProfiler/Widgets/SPacketContentView.cpp
@@ -783,8 +783,8 @@ void SPacketContentView::FilterNetId_OnTextCommitted(const FText& InNewText, ETe
 {
 	if (InNewText.IsNumeric())
 	{
-		int32 NewNetId = 0;
-		TTypeFromString<int32>::FromString(NewNetId, *InNewText.ToString());
+		uint64 NewNetId = 0;
+		TTypeFromString<uint64>::FromString(NewNetId, *InNewText.ToString());
 		SetFilterNetId(NewNetId);
 	}
 }
@@ -922,7 +922,7 @@ void SPacketContentView::SetPacket(uint32 InGameInstanceIndex, uint32 InConnecti
 
 ////////////////////////////////////////////////////////////////////////////////////////////////////
 
-void SPacketContentView::SetFilterNetId(const uint32 InNetId)
+void SPacketContentView::SetFilterNetId(const uint64 InNetId)
 {
 	FilterNetId = InNetId;
 
@@ -1056,12 +1056,12 @@ void SPacketContentView::UpdateState(float FontScale)
 						Name = NetProfilerName.Name;
 					});
 
-					uint32 NetId = 0;
+					uint64 NetId = 0;
 					if (Event.ObjectInstanceIndex != 0)
 					{
 						NetProfilerProvider->ReadObject(GameInstanceIndex, Event.ObjectInstanceIndex, [&NetId](const TraceServices::FNetProfilerObjectInstance& ObjectInstance)
 						{
-							NetId = ObjectInstance.NetId;
+							NetId = ObjectInstance.NetObjectId;
 						});
 					}
 
@@ -1139,7 +1139,7 @@ void SPacketContentView::UpdateHoveredEvent()
 
 		if (Event.ObjectInstanceIndex != 0)
 		{
-			Tooltip.AddNameValueTextLine(TEXT("Net Id:"), FText::AsNumber(ObjectInstance.NetId).ToString());
+			Tooltip.AddNameValueTextLine(TEXT("Net Id:"), FText::AsNumber(ObjectInstance.NetObjectId).ToString());
 			Tooltip.AddNameValueTextLine(TEXT("Type Id:"), FString::Printf(TEXT("0x%016" UINT64_x_FMT), ObjectInstance.TypeId));
 			Tooltip.AddNameValueTextLine(TEXT("Obj. LifeTime:"), FString::Format(TEXT("from {0} to {1}"),
 				{ TimeUtils::FormatTimeAuto(ObjectInstance.LifeTime.Begin), TimeUtils::FormatTimeAuto(ObjectInstance.LifeTime.End) }));
diff --git a/Engine/Source/Developer/TraceInsights/Private/Insights/NetworkingProfiler/Widgets/SPacketContentView.h b/Engine/Source/Developer/TraceInsights/Private/Insights/NetworkingProfiler/Widgets/SPacketContentView.h
--- a/Engine/Source/Developer/TraceInsights/Private/Insights/NetworkingProfiler/Widgets/SPacketContentView.h
+++ b/Engine/Source/Developer/TraceInsights/Private/Insights/NetworkingProfiler/Widgets/SPacketContentView.h
@@ -157,8 +157,8 @@ public:
 	void SetPacket(uint32 InGameInstanceIndex, uint32 InConnectionIndex, TraceServices::ENetProfilerConnectionMode InConnectionMode, uint32 InPacketIndex, int64 InPacketBitSize);
 
 	bool IsFilterByNetIdEnabled() const { return bFilterByNetId; }
-	uint32 GetFilterNetId() const { return FilterNetId; }
-	void SetFilterNetId(const uint32 InNetId);
+	uint64 GetFilterNetId() const { return FilterNetId; }
+	void SetFilterNetId(const uint64 InNetId);
 
 	bool IsFilterByEventTypeEnabled() const { return bFilterByEventType; }
 	uint32 GetFilterEventTypeIndex() const { return FilterEventTypeIndex; }
@@ -252,7 +252,7 @@ private:
 	FText FilterEventName;
 
 	bool bFilterByNetId;
-	uint32 FilterNetId;
+	uint64 FilterNetId;
 
 	bool bHighlightFilteredEvents;
 
diff --git a/Engine/Source/Developer/TraceInsights/Private/Insights/NetworkingProfiler/Widgets/SPacketView.cpp b/Engine/Source/Developer/TraceInsights/Private/Insights/NetworkingProfiler/Widgets/SPacketView.cpp
--- a/Engine/Source/Developer/TraceInsights/Private/Insights/NetworkingProfiler/Widgets/SPacketView.cpp
+++ b/Engine/Source/Developer/TraceInsights/Private/Insights/NetworkingProfiler/Widgets/SPacketView.cpp
@@ -279,7 +279,7 @@ void SPacketView::UpdateState()
 	struct FPacketFilter
 	{
 		bool bByNetId = false;
-		uint32 NetId = 0;
+		uint64 NetId = 0;
 		bool bByEventType = false;
 		uint32 EventTypeIndex = 0;
 		TraceServices::ENetProfilerAggregationMode AggregationMode = TraceServices::ENetProfilerAggregationMode::None;
@@ -368,12 +368,12 @@ void SPacketView::UpdateState()
 										}
 										if (bEventMatchesFilter && Filter.bByNetId)
 										{
-											uint32 NetId = uint32(-1);
+											uint64 NetId = uint64(-1);
 											if (Event.ObjectInstanceIndex != 0)
 											{
 												NetProfilerProvider->ReadObject(GameInstanceIndex, Event.ObjectInstanceIndex, [&NetId](const TraceServices::FNetProfilerObjectInstance& ObjectInstance)
 												{
-													NetId = ObjectInstance.NetId;
+													NetId = ObjectInstance.NetObjectId;
 												});
 											}
 											if (Filter.NetId != NetId)
diff --git a/Engine/Source/Developer/TraceServices/Private/Analyzers/NetTraceAnalyzer.cpp b/Engine/Source/Developer/TraceServices/Private/Analyzers/NetTraceAnalyzer.cpp
--- a/Engine/Source/Developer/TraceServices/Private/Analyzers/NetTraceAnalyzer.cpp
+++ b/Engine/Source/Developer/TraceServices/Private/Analyzers/NetTraceAnalyzer.cpp
@@ -758,7 +758,7 @@ void FNetTraceAnalyzer::HandleConnectionClosedEvent(const FOnEventContext& Conte
 void FNetTraceAnalyzer::HandleObjectCreatedEvent(const FOnEventContext& Context, const FEventData& EventData)
 {
 	const uint64 TypeId = EventData.GetValue<uint64>("TypeId");
-	const uint32 ObjectId = EventData.GetValue<uint32>("ObjectId");
+	const uint64 ObjectId = EventData.GetValue<uint64>("ObjectId");
 	const uint32 OwnerId = EventData.GetValue<uint32>("OwnerId");
 	const uint16 NameId = EventData.GetValue<uint16>("NameId");
 	const uint8 GameInstanceId = EventData.GetValue<uint8>("GameInstanceId");
@@ -775,7 +775,7 @@ void FNetTraceAnalyzer::HandleObjectCreatedEvent(const FOnEventContext& Context,
 			{
 				// Update existing object instance
 				ExistingInstance->LifeTime.Begin = GetLastTimestamp();
-				ExistingInstance->NetId = ObjectId;
+				ExistingInstance->NetObjectId = ObjectId;
 				ExistingInstance->TypeId = TypeId;
 
 				return;
@@ -793,7 +793,7 @@ void FNetTraceAnalyzer::HandleObjectCreatedEvent(const FOnEventContext& Context,
 	// Fill in object data
 	ObjectInstance.LifeTime.Begin = GetLastTimestamp();
 	ObjectInstance.NameIndex = IntCastChecked<uint16>(NameIndex);
-	ObjectInstance.NetId = ObjectId;
+	ObjectInstance.NetObjectId = ObjectId;
 	ObjectInstance.TypeId = TypeId;
 
 	// Add to active objects
@@ -804,7 +804,7 @@ void FNetTraceAnalyzer::HandleObjectDestroyedEvent(const FOnEventContext& Contex
 {
 	// Remove from active instances and mark the end timestamp in the persistent instance list
 	const uint8 GameInstanceId = EventData.GetValue<uint8>("GameInstanceId");
-	const uint32 ObjectId = EventData.GetValue<uint32>("ObjectId");
+	const uint64 ObjectId = EventData.GetValue<uint64>("ObjectId");
 
 	TSharedRef<FNetTraceGameInstanceState> GameInstanceState = GetOrCreateActiveGameInstanceState(GameInstanceId);
 
diff --git a/Engine/Source/Developer/TraceServices/Public/TraceServices/Model/NetProfiler.h b/Engine/Source/Developer/TraceServices/Public/TraceServices/Model/NetProfiler.h
--- a/Engine/Source/Developer/TraceServices/Public/TraceServices/Model/NetProfiler.h
+++ b/Engine/Source/Developer/TraceServices/Public/TraceServices/Model/NetProfiler.h
@@ -86,11 +86,16 @@ struct FNetProfilerLifeTime
 
 struct FNetProfilerObjectInstance
 {
-	uint32 ObjectIndex = 0U;		// Index in the object array
-	uint16 NameIndex = 0U;			// Index in the Name array
-	uint64 TypeId = uint64(0);		// ProtocolIdentifier
-	uint32 NetId = 0U;				// NetHandleIndex or NetGUID
-	FNetProfilerLifeTime LifeTime;	// Lifetime of this instance
+	FNetProfilerLifeTime LifeTime;			// Lifetime of this instance
+	union
+	{
+		UE_DEPRECATED(5.3, "Please use NetObjectId instead")
+		uint32 NetId;						// NetHandleIndex or NetGUID
+		uint64 NetObjectId = 0;				// NetHandleIndex or NetGUID
+	};
+	uint64 TypeId = uint64(0);				// ProtocolIdentifier
+	uint32 ObjectIndex = 0U;				// Index in the object array
+	uint16 NameIndex = 0U;					// Index in the Name array
 };
 
 union FNetProfilerBunchInfo
diff --git a/Engine/Source/Runtime/Experimental/Iris/Core/Public/Iris/ReplicationSystem/NetRefHandle.h b/Engine/Source/Runtime/Experimental/Iris/Core/Public/Iris/ReplicationSystem/NetRefHandle.h
--- a/Engine/Source/Runtime/Experimental/Iris/Core/Public/Iris/ReplicationSystem/NetRefHandle.h
+++ b/Engine/Source/Runtime/Experimental/Iris/Core/Public/Iris/ReplicationSystem/NetRefHandle.h
@@ -72,10 +72,9 @@ FORCEINLINE uint32 GetTypeHash(const FNetRefHandle& Handle)
 	return ::GetTypeHash(Handle.GetId());
 }
 
-FORCEINLINE uint32 GetObjectIdForNetTrace(const FNetRefHandle& Handle)
+FORCEINLINE uint64 GetObjectIdForNetTrace(const FNetRefHandle& Handle)
 {
-	// Until NetTrace is fixed: https://jira.it.epicgames.com/browse/UE-179190
-	return (uint32)Handle.GetId();
+	return Handle.GetId();
 }
 
 }
diff --git a/Engine/Source/Runtime/Net/Core/Private/Net/Core/Trace/NetTraceInternal.cpp b/Engine/Source/Runtime/Net/Core/Private/Net/Core/Trace/NetTraceInternal.cpp
--- a/Engine/Source/Runtime/Net/Core/Private/Net/Core/Trace/NetTraceInternal.cpp
+++ b/Engine/Source/Runtime/Net/Core/Private/Net/Core/Trace/NetTraceInternal.cpp
@@ -45,7 +45,7 @@ struct FNetTraceInternal
 	};
 
 	// Get next NameId used to track what we already have traced
-	uint32 static GetNextNameId();
+	UE::Net::FNetDebugNameId static GetNextNameId();
 
 	static inline FThreadBuffer* CreateThreadBuffer();
 
@@ -59,7 +59,7 @@ FNetTraceInternal::FThreadBuffer* FNetTraceInternal::CreateThreadBuffer()
 	return ThreadBuffer.Get();
 }
 
-uint32 FNetTraceInternal::GetNextNameId()
+UE::Net::FNetDebugNameId FNetTraceInternal::GetNextNameId()
 {
 	static std::atomic<UE::Net::FNetDebugNameId> NextNameId(1);
 
@@ -381,7 +381,7 @@ void FNetTrace::PopSendBunch(FNetTraceCollector& Collector)
 void FNetTrace::TraceCollectedEvents(FNetTraceCollector& Collector, uint32 GameInstanceId, uint32 ConnectionId, ENetTracePacketType PacketType)
 {
 	FNetTracePacketInfo PacketInfo;
-	PacketInfo.ConnectionId = ConnectionId;
+	PacketInfo.ConnectionId = (uint16)ConnectionId;
 	PacketInfo.GameInstanceId = GameInstanceId;
 	PacketInfo.PacketSequenceNumber = 0;
 	PacketInfo.PacketType = PacketType;
@@ -397,7 +397,7 @@ void FNetTrace::TracePacketDropped(uint32 GameInstanceId, uint32 ConnectionId, u
 	if (GNetTraceRuntimeVerbosity)
 	{
 		FNetTracePacketInfo PacketInfo;
-		PacketInfo.ConnectionId = ConnectionId;
+		PacketInfo.ConnectionId = (uint16)ConnectionId;
 		PacketInfo.GameInstanceId = GameInstanceId;
 		PacketInfo.PacketSequenceNumber = PacketSequenceNumber;
 		PacketInfo.PacketType = PacketType;
@@ -411,7 +411,7 @@ void FNetTrace::TracePacket(uint32 GameInstanceId, uint32 ConnectionId, uint32 P
 	if (GNetTraceRuntimeVerbosity)
 	{
 		FNetTracePacketInfo PacketInfo;
-		PacketInfo.ConnectionId = ConnectionId;
+		PacketInfo.ConnectionId = (uint16)ConnectionId;
 		PacketInfo.GameInstanceId = GameInstanceId;
 		PacketInfo.PacketSequenceNumber = PacketSequenceNumber;
 		PacketInfo.PacketType = PacketType;
@@ -422,7 +422,7 @@ void FNetTrace::TracePacket(uint32 GameInstanceId, uint32 ConnectionId, uint32 P
 	}
 }
 
-void FNetTrace::TraceObjectCreated(uint32 GameInstanceId, uint32 NetObjectId, const UE::Net::FNetDebugName* DebugName, uint64 TypeIdentifier, uint32 OwnerId)
+void FNetTrace::TraceObjectCreated(uint32 GameInstanceId, uint64 NetObjectId, const UE::Net::FNetDebugName* DebugName, uint64 TypeIdentifier, uint32 OwnerId)
 {
 	if (!GNetTraceRuntimeVerbosity)
 	{
@@ -436,7 +436,7 @@ void FNetTrace::TraceObjectCreated(uint32 GameInstanceId, uint32 NetObjectId, co
 	FNetTraceInternal::Reporter::ReportObjectCreated(GameInstanceId, NetObjectId, DebugName->DebugNameId, TypeIdentifier, OwnerId);
 }
 
-void FNetTrace::TraceObjectCreated(uint32 GameInstanceId, uint32 NetObjectId, const FName ObjectName, uint64 TypeIdentifier, uint32 OwnerId)
+void FNetTrace::TraceObjectCreated(uint32 GameInstanceId, uint64 NetObjectId, const FName ObjectName, uint64 TypeIdentifier, uint32 OwnerId)
 {
 	if (!GNetTraceRuntimeVerbosity)
 	{
@@ -446,7 +446,7 @@ void FNetTrace::TraceObjectCreated(uint32 GameInstanceId, uint32 NetObjectId, co
 	FNetTraceInternal::Reporter::ReportObjectCreated(GameInstanceId, NetObjectId, FNetTrace::TraceName(ObjectName), TypeIdentifier, OwnerId);
 }
 
-void FNetTrace::TraceObjectDestroyed(uint32 GameInstanceId, uint32 NetObjectId)
+void FNetTrace::TraceObjectDestroyed(uint32 GameInstanceId, uint64 NetObjectId)
 {
 	if (GNetTraceRuntimeVerbosity)
 	{
diff --git a/Engine/Source/Runtime/Net/Core/Private/Net/Core/Trace/Reporters/NetTraceReporter.cpp b/Engine/Source/Runtime/Net/Core/Private/Net/Core/Trace/Reporters/NetTraceReporter.cpp
--- a/Engine/Source/Runtime/Net/Core/Private/Net/Core/Trace/Reporters/NetTraceReporter.cpp
+++ b/Engine/Source/Runtime/Net/Core/Private/Net/Core/Trace/Reporters/NetTraceReporter.cpp
@@ -30,14 +30,14 @@ UE_TRACE_EVENT_END()
 
 UE_TRACE_EVENT_BEGIN(NetTrace, ObjectCreatedEvent)	
 	UE_TRACE_EVENT_FIELD(uint64, TypeId)
-	UE_TRACE_EVENT_FIELD(uint32, ObjectId)
+	UE_TRACE_EVENT_FIELD(uint64, ObjectId)
 	UE_TRACE_EVENT_FIELD(uint32, OwnerId)
 	UE_TRACE_EVENT_FIELD(uint16, NameId)
 	UE_TRACE_EVENT_FIELD(uint8, GameInstanceId)
 UE_TRACE_EVENT_END()
 
 UE_TRACE_EVENT_BEGIN(NetTrace, ObjectDestroyedEvent)
-	UE_TRACE_EVENT_FIELD(uint32, ObjectId)
+	UE_TRACE_EVENT_FIELD(uint64, ObjectId)
 	UE_TRACE_EVENT_FIELD(uint8, GameInstanceId)
 UE_TRACE_EVENT_END()
 
@@ -131,7 +131,7 @@ void FNetTraceReporter::ReportInitEvent(uint32 NetTraceVersion)
 void FNetTraceReporter::ReportInstanceUpdated(uint32 GameInstanceId, bool bIsServer, const TCHAR* Name)
 {
 	UE_TRACE_LOG(NetTrace, InstanceUpdatedEvent, NetChannel)
-		<< InstanceUpdatedEvent.GameInstanceId(GameInstanceId)
+		<< InstanceUpdatedEvent.GameInstanceId((uint8)GameInstanceId)
 		<< InstanceUpdatedEvent.bIsServer(bIsServer)
 		<< InstanceUpdatedEvent.Name(Name);
 }
@@ -139,7 +139,7 @@ void FNetTraceReporter::ReportInstanceUpdated(uint32 GameInstanceId, bool bIsSer
 void FNetTraceReporter::ReportInstanceDestroyed(uint32 GameInstanceId)
 {
 	UE_TRACE_LOG(NetTrace, InstanceDestroyedEvent, NetChannel)
-		<< InstanceDestroyedEvent.GameInstanceId(GameInstanceId);
+		<< InstanceDestroyedEvent.GameInstanceId((uint8)GameInstanceId);
 }
 
 void FNetTraceReporter::ReportAnsiName(UE::Net::FNetDebugNameId NameId, uint32 NameSize, const char* Name)
@@ -155,7 +155,7 @@ void FNetTraceReporter::ReportPacketDropped(const FNetTracePacketInfo& PacketInf
 		<< PacketDroppedEvent.Timestamp(FPlatformTime::Cycles64())
 		<< PacketDroppedEvent.SequenceNumber(PacketInfo.PacketSequenceNumber)
 		<< PacketDroppedEvent.ConnectionId(PacketInfo.ConnectionId)
-		<< PacketDroppedEvent.GameInstanceId(PacketInfo.GameInstanceId)
+		<< PacketDroppedEvent.GameInstanceId((uint8)PacketInfo.GameInstanceId)
 		<< PacketDroppedEvent.PacketType((uint8)PacketInfo.PacketType);
 }
 
@@ -166,7 +166,7 @@ void FNetTraceReporter::ReportPacket(const FNetTracePacketInfo& PacketInfo, uint
 		<< PacketEvent.PacketBits(PacketBits)
 		<< PacketEvent.SequenceNumber(PacketInfo.PacketSequenceNumber)
 		<< PacketEvent.ConnectionId(PacketInfo.ConnectionId)
-		<< PacketEvent.GameInstanceId(PacketInfo.GameInstanceId)
+		<< PacketEvent.GameInstanceId((uint8)PacketInfo.GameInstanceId)
 		<< PacketEvent.PacketType((uint8)PacketInfo.PacketType);
 }
 
@@ -186,7 +186,7 @@ void FNetTraceReporter::ReportPacketContent(FNetTracePacketContentEvent* Events,
 	{
 		UE_TRACE_LOG(NetTrace, PacketContentEvent, NetChannel)
 			<< PacketContentEvent.ConnectionId(InPacketInfo.ConnectionId)
-			<< PacketContentEvent.GameInstanceId(InPacketInfo.GameInstanceId)
+			<< PacketContentEvent.GameInstanceId((uint8)InPacketInfo.GameInstanceId)
 			<< PacketContentEvent.PacketType((uint8)InPacketInfo.PacketType)
 			<< PacketContentEvent.Data(InBuffer, Count);
 	};
@@ -196,7 +196,7 @@ void FNetTraceReporter::ReportPacketContent(FNetTracePacketContentEvent* Events,
 		// Flush
 		if ((BufferPtr - Buffer) > FlushBufferThreshold)
 		{
-			FlushPacketContentBuffer(PacketInfo, Buffer, BufferPtr - Buffer);
+			FlushPacketContentBuffer(PacketInfo, Buffer, (uint32)(BufferPtr - Buffer));
 			BufferPtr = Buffer;
 			LastOffset = 0;
 		}
@@ -209,13 +209,29 @@ void FNetTraceReporter::ReportPacketContent(FNetTracePacketContentEvent* Events,
 		switch (ENetTracePacketContentEventType(CurrentEvent.EventType))
 		{
 			case ENetTracePacketContentEventType::Object:
+			{
+				// NestingLevel
+				*(BufferPtr++) = CurrentEvent.NestingLevel;
+
+				FTraceUtils::Encode7bit(CurrentEvent.ObjectId, BufferPtr);
+
+				// Encode event data, all offsets are delta compressed against previous begin marker
+				const uint64 StartPos = CurrentEvent.StartPos;
+
+				// Start
+				FTraceUtils::Encode7bit(StartPos - LastOffset, BufferPtr);
+				LastOffset = StartPos;
+
+				// End
+				FTraceUtils::Encode7bit(CurrentEvent.EndPos - StartPos, BufferPtr);
+			}
+			break;
 			case ENetTracePacketContentEventType::NameId:
 			{
 				// NestingLevel
 				*(BufferPtr++) = CurrentEvent.NestingLevel;
 
-				uint32 EventId = (ENetTracePacketContentEventType)CurrentEvent.EventType == ENetTracePacketContentEventType::Object ? CurrentEvent.ObjectId : CurrentEvent.DebugNameId;
-				FTraceUtils::Encode7bit(EventId, BufferPtr);
+				FTraceUtils::Encode7bit(CurrentEvent.DebugNameId, BufferPtr);
 
 				// Encode event data, all offsets are delta compressed against previous begin marker
 				const uint64 StartPos = CurrentEvent.StartPos;
@@ -266,7 +282,7 @@ void FNetTraceReporter::ReportPacketContent(FNetTracePacketContentEvent* Events,
 
 	if (BufferPtr > Buffer)
 	{
-		FlushPacketContentBuffer(PacketInfo, Buffer, BufferPtr - Buffer);
+		FlushPacketContentBuffer(PacketInfo, Buffer, (uint32)(BufferPtr - Buffer));
 		BufferPtr = Buffer;
 	}
 }
@@ -274,23 +290,23 @@ void FNetTraceReporter::ReportPacketContent(FNetTracePacketContentEvent* Events,
 void FNetTraceReporter::ReportConnectionCreated(uint32 GameInstanceId, uint32 ConnectionId)
 {
 	UE_TRACE_LOG(NetTrace, ConnectionCreatedEvent, NetChannel)
-		<< ConnectionCreatedEvent.ConnectionId(ConnectionId)
-		<< ConnectionCreatedEvent.GameInstanceId(GameInstanceId);
+		<< ConnectionCreatedEvent.ConnectionId((uint16)ConnectionId)
+		<< ConnectionCreatedEvent.GameInstanceId((uint8)GameInstanceId);
 }
 
 void FNetTraceReporter::ReportConnectionStateUpdated(uint32 GameInstanceId, uint32 ConnectionId, uint8 ConnectionStateValue)
 {
 	UE_TRACE_LOG(NetTrace, ConnectionStateUpdatedEvent, NetChannel)
-		<< ConnectionStateUpdatedEvent.ConnectionId(ConnectionId)
-		<< ConnectionStateUpdatedEvent.GameInstanceId(GameInstanceId)
+		<< ConnectionStateUpdatedEvent.ConnectionId((uint16)ConnectionId)
+		<< ConnectionStateUpdatedEvent.GameInstanceId((uint8)GameInstanceId)
 		<< ConnectionStateUpdatedEvent.ConnectionStateValue(ConnectionStateValue);
 }
 
 void FNetTraceReporter::ReportConnectionUpdated(uint32 GameInstanceId, uint32 ConnectionId, const TCHAR* AddressString, const TCHAR* OwningActor)
 {
 	UE_TRACE_LOG(NetTrace, ConnectionUpdatedEvent, NetChannel)
-		<< ConnectionUpdatedEvent.GameInstanceId(GameInstanceId)
-		<< ConnectionUpdatedEvent.ConnectionId(ConnectionId)
+		<< ConnectionUpdatedEvent.GameInstanceId((uint8)GameInstanceId)
+		<< ConnectionUpdatedEvent.ConnectionId((uint16)ConnectionId)
 		<< ConnectionUpdatedEvent.Name(OwningActor)
 		<< ConnectionUpdatedEvent.Address(AddressString);
 }
@@ -298,8 +314,8 @@ void FNetTraceReporter::ReportConnectionUpdated(uint32 GameInstanceId, uint32 Co
 void FNetTraceReporter::ReportConnectionClosed(uint32 GameInstanceId, uint32 ConnectionId)
 {
 	UE_TRACE_LOG(NetTrace, ConnectionClosedEvent, NetChannel)
-		<< ConnectionClosedEvent.ConnectionId(ConnectionId)
-		<< ConnectionClosedEvent.GameInstanceId(GameInstanceId);
+		<< ConnectionClosedEvent.ConnectionId((uint16)ConnectionId)
+		<< ConnectionClosedEvent.GameInstanceId((uint8)GameInstanceId);
 }
 
 void FNetTraceReporter::ReportPacketStatsCounter(uint32 GameInstanceId, uint32 ConnectionId, UE::Net::FNetDebugNameId CounterNameId, uint32 StatValue)
@@ -307,8 +323,8 @@ void FNetTraceReporter::ReportPacketStatsCounter(uint32 GameInstanceId, uint32 C
 	UE_TRACE_LOG(NetTrace, PacketStatsCounterEvent, NetChannel)
 		<< PacketStatsCounterEvent.StatsValue(StatValue)
 		<< PacketStatsCounterEvent.NameId(CounterNameId)
-		<< PacketStatsCounterEvent.ConnectionId(ConnectionId)
-		<< PacketStatsCounterEvent.GameInstanceId(GameInstanceId);
+		<< PacketStatsCounterEvent.ConnectionId((uint16)ConnectionId)
+		<< PacketStatsCounterEvent.GameInstanceId((uint8)GameInstanceId);
 }
 
 void FNetTraceReporter::ReportFrameStatsCounter(uint32 GameInstanceId, UE::Net::FNetDebugNameId CounterNameId, uint32 StatValue)
@@ -317,24 +333,24 @@ void FNetTraceReporter::ReportFrameStatsCounter(uint32 GameInstanceId, UE::Net::
 		<< FrameStatsCounterEvent.Timestamp(FPlatformTime::Cycles64())
 		<< FrameStatsCounterEvent.StatsValue(StatValue)
 		<< FrameStatsCounterEvent.NameId(CounterNameId)
-		<< FrameStatsCounterEvent.GameInstanceId(GameInstanceId);
+		<< FrameStatsCounterEvent.GameInstanceId((uint8)GameInstanceId);
 }
 
-void FNetTraceReporter::ReportObjectCreated(uint32 GameInstanceId, uint32 NetObjectId, UE::Net::FNetDebugNameId NameId, uint64 TypeIdentifier, uint32 OwnerId)
+void FNetTraceReporter::ReportObjectCreated(uint32 GameInstanceId, uint64 NetObjectId, UE::Net::FNetDebugNameId NameId, uint64 TypeIdentifier, uint32 OwnerId)
 {
 	UE_TRACE_LOG(NetTrace, ObjectCreatedEvent, NetChannel)
 		<< ObjectCreatedEvent.TypeId(TypeIdentifier)
 		<< ObjectCreatedEvent.ObjectId(NetObjectId)
 		<< ObjectCreatedEvent.OwnerId(OwnerId)
 		<< ObjectCreatedEvent.NameId(NameId)
-		<< ObjectCreatedEvent.GameInstanceId(GameInstanceId);
+		<< ObjectCreatedEvent.GameInstanceId((uint8)GameInstanceId);
 }
 
-void FNetTraceReporter::ReportObjectDestroyed(uint32 GameInstanceId, uint32 NetObjectId)
+void FNetTraceReporter::ReportObjectDestroyed(uint32 GameInstanceId, uint64 NetObjectId)
 {
 	UE_TRACE_LOG(NetTrace, ObjectDestroyedEvent, NetChannel)
 		<< ObjectDestroyedEvent.ObjectId(NetObjectId)
-		<< ObjectDestroyedEvent.GameInstanceId(GameInstanceId);
+		<< ObjectDestroyedEvent.GameInstanceId((uint8)GameInstanceId);
 }
 
 #endif
diff --git a/Engine/Source/Runtime/Net/Core/Private/Net/Core/Trace/Reporters/NetTraceReporter.h b/Engine/Source/Runtime/Net/Core/Private/Net/Core/Trace/Reporters/NetTraceReporter.h
--- a/Engine/Source/Runtime/Net/Core/Private/Net/Core/Trace/Reporters/NetTraceReporter.h
+++ b/Engine/Source/Runtime/Net/Core/Private/Net/Core/Trace/Reporters/NetTraceReporter.h
@@ -15,8 +15,8 @@ struct FNetTraceReporter
 	static void ReportPacketDropped(const FNetTracePacketInfo& PacketInfo);
 	static void ReportPacket(const FNetTracePacketInfo& PacketInfo, uint32 PacketBits);
 	static void ReportAnsiName(UE::Net::FNetDebugNameId NameId, uint32 NameSize, const char* Name);	
-	static void ReportObjectCreated(uint32 GameInstanceId, uint32 NetObjectId, UE::Net::FNetDebugNameId NameId, uint64 TypeIdentifier, uint32 OwnerId);
-	static void ReportObjectDestroyed(uint32 GameInstanceId, uint32 NetObjectId);
+	static void ReportObjectCreated(uint32 GameInstanceId, uint64 NetObjectId, UE::Net::FNetDebugNameId NameId, uint64 TypeIdentifier, uint32 OwnerId);
+	static void ReportObjectDestroyed(uint32 GameInstanceId, uint64 NetObjectId);
 	static void ReportConnectionCreated(uint32 GameInstanceId, uint32 ConnectionId);
 	static void ReportConnectionStateUpdated(uint32 GameInstanceId, uint32 ConnectionId, uint8 ConnectionStateValue);
 	static void ReportConnectionUpdated(uint32 GameInstanceId, uint32 ConnectionId, const TCHAR* AddressString, const TCHAR* OwningActor);
diff --git a/Engine/Source/Runtime/Net/Core/Public/Net/Core/Trace/Private/NetTraceInternal.h b/Engine/Source/Runtime/Net/Core/Public/Net/Core/Trace/Private/NetTraceInternal.h
--- a/Engine/Source/Runtime/Net/Core/Public/Net/Core/Trace/Private/NetTraceInternal.h
+++ b/Engine/Source/Runtime/Net/Core/Public/Net/Core/Trace/Private/NetTraceInternal.h
@@ -92,7 +92,7 @@ struct FNetTracePacketContentEvent
 	union
 	{
 		UE::Net::FNetDebugNameId DebugNameId;
-		uint32 ObjectId;
+		uint64 ObjectId;
 		FNetTraceBunchInfo BunchInfo;
 	};
 	struct
@@ -117,7 +117,7 @@ class FNetTraceEventScope
 	UE_NONCOPYABLE(FNetTraceEventScope)
 public:
 
-	FNetTraceEventScope(uint32 ObjectID, T& InStream, FNetTraceCollector* InCollector, uint32 Verbosity);
+	FNetTraceEventScope(uint64 ObjectId, T& InStream, FNetTraceCollector* InCollector, uint32 Verbosity);
 
 	// When a scope is disabled due to verbosity we want to deffer the tracing of names
 	// We do this by using a lambda, this also allows us to complete eliminate statics if a scope is compiled out due to compile time verbosity
@@ -127,7 +127,7 @@ public:
 	~FNetTraceEventScope();
 
 	void SetEventName(UE::Net::FNetDebugNameId NetTraceNameId);
-	void SetEventObjectId(uint32 ObjectId);
+	void SetEventObjectId(uint64 ObjectId);
 	void ExitScope();
 	bool IsValid() const;
 
@@ -151,7 +151,7 @@ public:
 	~FNetTraceNullEventScope() {};
 
 	void SetEventName(UE::Net::FNetDebugNameId NetTraceNameId) {};
-	void SetEventObjectId(uint32 ObjectId) {};
+	void SetEventObjectId(uint64 ObjectId) {};
 	void ExitScope() {};
 	bool IsValid() const { return false; }
 };
@@ -231,13 +231,13 @@ struct FNetTrace
  	NETCORE_API static void TracePacket(uint32 GameInstanceId, uint32 ConnectionId, uint32 PacketSequenceNumber, uint32 PacketBits, ENetTracePacketType PacketType);
 
 	/** Trace that we have created a new object with the given ObjectId */
-	NETCORE_API static void TraceObjectCreated(uint32 GameInstanceId, uint32 NetObjectId, const UE::Net::FNetDebugName* Name, uint64 TypeIdentifier, uint32 OwnerId);
+	NETCORE_API static void TraceObjectCreated(uint32 GameInstanceId, uint64 NetObjectId, const UE::Net::FNetDebugName* Name, uint64 TypeIdentifier, uint32 OwnerId);
 
 	/** Trace that we have created a new object with the given ObjectId */
-	NETCORE_API static void TraceObjectCreated(uint32 GameInstanceId, uint32 NetObjectId, const FName Name, uint64 TypeIdentifier, uint32 OwnerId);
+	NETCORE_API static void TraceObjectCreated(uint32 GameInstanceId, uint64 NetObjectId, const FName Name, uint64 TypeIdentifier, uint32 OwnerId);
 
 	/** Trace that a handle has been destroyed, the handle contains required info */
-	NETCORE_API static void TraceObjectDestroyed(uint32 GameInstanceId, uint32 NetObjectId);
+	NETCORE_API static void TraceObjectDestroyed(uint32 GameInstanceId, uint64 NetObjectId);
 
 	/** Trace that we have added a new connection for the given GameInstanceId */
 	NETCORE_API static void TraceConnectionCreated(uint32 GameInstanceId, uint32 ConnectionId);
@@ -279,14 +279,14 @@ struct FNetTrace
 };
 
 template<typename T>
-uint32 GetObjectIdForNetTrace(const T&)
+uint64 GetObjectIdForNetTrace(const T&)
 {
-	static_assert(sizeof(T) == 0, "Not supported type for NetTraceObjectID, implement uint32 GetObjectIdForNetTrace(const T&)");
+	static_assert(sizeof(T) == 0, "Not supported type for NetTraceObjectID, implement uint64 GetObjectIdForNetTrace(const T&)");
 }
 
-inline uint32 GetObjectIdForNetTrace(const FNetworkGUID& NetGUID)
+inline uint64 GetObjectIdForNetTrace(const FNetworkGUID& NetGUID)
 {
-	return (uint32)NetGUID.ObjectId;
+	return NetGUID.ObjectId;
 }
 
 template<typename T>
@@ -409,10 +409,10 @@ FNetTraceCollector* FNetTrace::GetCollectorAtVerbosity(FNetTraceCollector* Colle
 
 // FNetTraceEventScope implementation
 template <typename T>
-FNetTraceEventScope<T>::FNetTraceEventScope(uint32 InObjectId, T& InStream, FNetTraceCollector* InCollector, uint32 Verbosity)
-: Collector(FNetTrace::GetCollectorAtVerbosity(InCollector, Verbosity))
-, Stream(InStream)
-, EventIndex(FNetTrace::InvalidEventIndex)
+FNetTraceEventScope<T>::FNetTraceEventScope(uint64 InObjectId, T& InStream, FNetTraceCollector* InCollector, uint32 Verbosity)
+	: Collector(FNetTrace::GetCollectorAtVerbosity(InCollector, Verbosity))
+	, Stream(InStream)
+	, EventIndex(FNetTrace::InvalidEventIndex)
 {
 	if (Collector)
 	{
@@ -464,13 +464,14 @@ void FNetTraceEventScope<T>::SetEventName(UE::Net::FNetDebugNameId NetTraceNameI
 }
 
 template <typename T>
-void FNetTraceEventScope<T>::SetEventObjectId(uint32 ObjectId)
+void FNetTraceEventScope<T>::SetEventObjectId(uint64 ObjectId)
 {
 	if (EventIndex != FNetTrace::InvalidEventIndex)
 	{
 		Collector->EditEventNoCheck(EventIndex).ObjectId = ObjectId;
 	}
 }
+
 template <typename T>
 bool FNetTraceEventScope<T>::IsValid() const
 {
@@ -579,8 +580,8 @@ void FNetTraceCollector::Reset()
 
 #define UE_NET_TRACE_INTERNAL_DYNAMIC_NAME(Name, Collector, StartPos, EndPos, Verbosity) UE_NET_TRACE_DO_IF(FNetTrace::GetCollectorAtVerbosity(Collector, Verbosity), FNetTrace::TracePacketContentEvent(*Collector, FNetTrace::TraceName(Name), StartPos, EndPos, Verbosity))
 #define UE_NET_TRACE_INTERNAL_ASSIGNED_GUID(GameInstanceId, NetGUID, PathName, OwnerId) UE_NET_TRACE_DO_IF(GNetTraceRuntimeVerbosity, FNetTrace::TraceObjectCreated(GameInstanceId, GetObjectIdForNetTrace(NetGUID), PathName, 0U, OwnerId))
-#define UE_NET_TRACE_INTERNAL_NETHANDLE_CREATED(Handle, DebugName, ProtocolId, OwnerId) FNetTrace::TraceObjectCreated(Handle.GetReplicationSystemId(), Handle.GetId(), DebugName, ProtocolId, OwnerId)
-#define UE_NET_TRACE_INTERNAL_NETHANDLE_DESTROYED(Handle) FNetTrace::TraceObjectDestroyed(Handle.GetReplicationSystemId(), Handle.GetId())
+#define UE_NET_TRACE_INTERNAL_NETHANDLE_CREATED(Handle, DebugName, ProtocolId, OwnerId) FNetTrace::TraceObjectCreated(Handle.GetReplicationSystemId(), GetObjectIdForNetTrace(Handle), DebugName, ProtocolId, OwnerId)
+#define UE_NET_TRACE_INTERNAL_NETHANDLE_DESTROYED(Handle) FNetTrace::TraceObjectDestroyed(Handle.GetReplicationSystemId(), GetObjectIdForNetTrace(Handle))
 #define UE_NET_TRACE_INTERNAL_CONNECTION_CREATED(...) FNetTrace::TraceConnectionCreated(__VA_ARGS__)
 #define UE_NET_TRACE_INTERNAL_CONNECTION_STATE_UPDATED(...) FNetTrace::TraceConnectionStateUpdated(__VA_ARGS__)
 #define UE_NET_TRACE_INTERNAL_CONNECTION_UPDATED(...) FNetTrace::TraceConnectionUpdated(__VA_ARGS__)
