Enabling UnsafeTypeCastWarningLevel for NetCore

diff --git a/Engine/Source/Runtime/Net/Core/NetCore.Build.cs b/Engine/Source/Runtime/Net/Core/NetCore.Build.cs
--- a/Engine/Source/Runtime/Net/Core/NetCore.Build.cs
+++ b/Engine/Source/Runtime/Net/Core/NetCore.Build.cs
@@ -17,5 +17,7 @@ public class NetCore : ModuleRules
 				"NetCommon"
 			}
 		);
+
+		UnsafeTypeCastWarningLevel = WarningLevel.Error;
 	}
 }
diff --git a/Engine/Source/Runtime/Net/Core/Private/Net/Core/Connection/EscalationStates.cpp b/Engine/Source/Runtime/Net/Core/Private/Net/Core/Connection/EscalationStates.cpp
--- a/Engine/Source/Runtime/Net/Core/Private/Net/Core/Connection/EscalationStates.cpp
+++ b/Engine/Source/Runtime/Net/Core/Private/Net/Core/Connection/EscalationStates.cpp
@@ -23,14 +23,15 @@ void FEscalationState::ValidateConfigInternal()
 
 void FEscalationState::ValidateTimePeriod(int8& Value, const TCHAR* PropertyName, EValidateTime Requirement/*=EValidateTime::Optional*/)
 {
-	int32 BaseValue = (Requirement == EValidateTime::MustBeSet) ? 1 : 0;
+	const int8 BaseValue = (Requirement == EValidateTime::MustBeSet) ? 1 : 0;
 
 	if (Value < BaseValue || Value > 16)
 	{
 		UE_LOG(LogNetCore, Warning, TEXT("FEscalationState: %s '%i' must be between %i and 16. Clamping."),
 				PropertyName, Value, BaseValue);
 
-		Value = FMath::Clamp((int32)Value, BaseValue, 16);
+		constexpr int8 MaxValue = 16;
+		Value = FMath::Clamp(Value, BaseValue, MaxValue);
 	}
 }
 
@@ -159,7 +160,8 @@ EEscalateResult FEscalationManager::UpdateSeverity(ESeverityUpdate Update, EEsca
 	bool bEscalate = Update == ESeverityUpdate::Escalate || Update == ESeverityUpdate::AutoEscalate;
 	EEscalateResult ReturnVal = bEscalate ? EEscalateResult::Escalated : EEscalateResult::Deescalated;
 	const TArray<TStructOnScope<FEscalationState>>& EscalationSeverity = BaseConfig->EscalationSeverityState;
-	int32 NewStateIdx = FMath::Clamp(ActiveState + (bEscalate ? 1 : -1), 0, EscalationSeverity.Num()-1);
+	
+	int8 NewStateIdx = static_cast<int8>(FMath::Clamp(ActiveState + (bEscalate ? 1 : -1), 0, EscalationSeverity.Num()-1));
 
 	if (NewStateIdx != ActiveState)
 	{
@@ -265,7 +267,7 @@ EEscalateResult FEscalationManager::UpdateSeverity(ESeverityUpdate Update, EEsca
 			if (bEscalate && State->HasHitAnyQuota({RegisteredCounters, CountersPerPeriodHistory, MakeArrayView(SecondCounters),
 													MakeArrayView(FrameCounters), EQuotaType::EscalateQuota}))
 			{
-				NewStateIdx = FMath::Clamp(ActiveState + 1, 0, EscalationSeverity.Num()-1);
+				NewStateIdx = static_cast<int8>(FMath::Clamp(ActiveState + 1, 0, EscalationSeverity.Num()-1));
 
 				if (NewStateIdx == ActiveState)
 				{
diff --git a/Engine/Source/Runtime/Net/Core/Private/Net/Core/Misc/DDoSDetection.cpp b/Engine/Source/Runtime/Net/Core/Private/Net/Core/Misc/DDoSDetection.cpp
--- a/Engine/Source/Runtime/Net/Core/Private/Net/Core/Misc/DDoSDetection.cpp
+++ b/Engine/Source/Runtime/Net/Core/Private/Net/Core/Misc/DDoSDetection.cpp
@@ -5,7 +5,7 @@
 #include "Net/Core/Misc/DDoSDetection.h"
 #include "Net/Core/Misc/NetCoreLog.h"
 #include "Misc/ConfigCacheIni.h"
-
+#include "Templates/UnrealTemplate.h"
 
 /**
  * DDoS Detection
@@ -108,7 +108,7 @@ void FDDoSDetection::InitConfig()
 	GConfig->GetInt(DDoSSection, TEXT("HitchTimeQuotaMS"), HitchTimeQuotaMS, GEngineIni);
 	GConfig->GetInt(DDoSSection, TEXT("HitchFrameTolerance"), HitchFrameTolerance32, GEngineIni);
 
-	HitchFrameTolerance = HitchFrameTolerance32;
+	HitchFrameTolerance = IntCastChecked<int8>(HitchFrameTolerance32);
 	DDoSLogSpamLimit = DDoSLogSpamLimit > 0 ? DDoSLogSpamLimit : 64;
 
 	DetectionSeverity.Empty();
@@ -144,7 +144,7 @@ void FDDoSDetection::InitConfig()
 
 				if (GConfig->GetInt(*CurSection, TEXT("EscalateTimeQuotaMSPerFrame"), EscalateTime32, GEngineIni))
 				{
-					CurState.EscalateTimeQuotaMSPerFrame = EscalateTime32;
+					CurState.EscalateTimeQuotaMSPerFrame = IntCastChecked<int16>(EscalateTime32);
 				}
 
 				HighestCooloffTime = FMath::Max(HighestCooloffTime, CurState.CooloffTime);
@@ -172,7 +172,7 @@ void FDDoSDetection::InitConfig()
 
 void FDDoSDetection::UpdateSeverity(bool bEscalate)
 {
-	int32 NewState = FMath::Clamp(ActiveState + (bEscalate ? 1 : -1), 0, DetectionSeverity.Num()-1);
+	int8 NewState = static_cast<int8>(FMath::Clamp(ActiveState + (bEscalate ? 1 : -1), 0, DetectionSeverity.Num()-1));
 
 	if (NewState != ActiveState)
 	{
@@ -263,7 +263,7 @@ void FDDoSDetection::PreFrameReceive(float DeltaTime)
 
 		if (HitchTimeQuotaMS > 0 && EndFrameRecvTimestamp != 0.0)
 		{
-			double HitchTimeMS = (StartFrameRecvTimestamp - EndFrameRecvTimestamp) * 1000.0;
+			const double HitchTimeMS = (StartFrameRecvTimestamp - EndFrameRecvTimestamp) * 1000.0;
 
 			if ((((int32)HitchTimeMS) - HitchTimeQuotaMS) > 0)
 			{
@@ -279,7 +279,7 @@ void FDDoSDetection::PreFrameReceive(float DeltaTime)
 		}
 
 		// At the start of every frame, adjust the DDoS detection based upon DeltaTime - unless there is excessive hitching
-		FrameAdjustment = (HitchFrameCount > 0 && HitchFrameCount > HitchFrameTolerance) ? 1.f : (double)DeltaTime / ExpectedFrameTime;
+		FrameAdjustment = FloatCastChecked<float>((HitchFrameCount > 0 && HitchFrameCount > HitchFrameTolerance) ? 1.0 : (double)DeltaTime / ExpectedFrameTime, UE::LWC::DefaultFloatPrecision);
 
 		if (ActiveState > 0 && CooloffTime > 0 && (float)(StartFrameRecvTimestamp - LastMetEscalationConditions) > (float)CooloffTime)
 		{
diff --git a/Engine/Source/Runtime/Net/Core/Private/Net/Core/Misc/PacketAudit.cpp b/Engine/Source/Runtime/Net/Core/Private/Net/Core/Misc/PacketAudit.cpp
--- a/Engine/Source/Runtime/Net/Core/Private/Net/Core/Misc/PacketAudit.cpp
+++ b/Engine/Source/Runtime/Net/Core/Private/Net/Core/Misc/PacketAudit.cpp
@@ -216,7 +216,7 @@ public:
 
 			Offset += Num;
 
-			*MemSizeData = Offset;
+			*MemSizeData = IntCastChecked<uint32>(Offset);
 		}
 	}
 
@@ -326,7 +326,7 @@ void FPacketAudit::AddStage_Internal(FString StageName, FBitWriter& OutPacket, b
 
 	FPacketStageData& StageData = GPendingSendPacket.StageMap.Add(StageName);
 
-	StageData.SizeBits = (bByteAligned ? (OutPacket.GetNumBytes() * 8) : OutPacket.GetNumBits());
+	StageData.SizeBits = IntCastChecked<uint32>(bByteAligned ? (OutPacket.GetNumBytes() * 8) : OutPacket.GetNumBits());
 	StageData.StageCRC = PacketCRC(OutPacket.GetData(), StageData.SizeBits);
 #endif
 }
@@ -344,7 +344,7 @@ void FPacketAudit::CheckStage_Internal(FString StageName, FBitReader& InPacket,
 
 			if (StageData != nullptr)
 			{
-				uint32 BitsLeft = (bByteAligned ? ((InPacket.GetNumBytes() * 8) - InPacket.GetPosBits()) : InPacket.GetBitsLeft());
+				uint32 BitsLeft = IntCastChecked<uint8>(bByteAligned ? ((InPacket.GetNumBytes() * 8) - InPacket.GetPosBits()) : InPacket.GetBitsLeft());
 
 				if (BitsLeft == StageData->SizeBits)
 				{
@@ -406,8 +406,7 @@ void FPacketAudit::NotifyLowLevelSend_Internal(FBitWriter& OutPacket)
 		// @todo #JohnB: Restore when you have the whole packet pipeline unified
 		//check(GPendingSendPacket.Writer == &OutPacket);
 
-
-		uint32 OutSizeBits = OutPacket.GetNumBits();
+		uint32 OutSizeBits = IntCastChecked<uint32>(OutPacket.GetNumBits());
 
 		if (OutSizeBits > 0)
 		{
@@ -439,7 +438,7 @@ void FPacketAudit::NotifyLowLevelSend_Internal(FBitWriter& OutPacket)
 void FPacketAudit::NotifyLowLevelReceive_Internal(FBitReader& InPacket)
 {
 #if !UE_BUILD_SHIPPING
-	GCurrentReceivePacket = PacketCRC(InPacket.GetData(), InPacket.GetNumBits());
+	GCurrentReceivePacket = PacketCRC(InPacket.GetData(), IntCastChecked<uint32>(InPacket.GetNumBits()));
 	GLastReceivePacket = 0;
 #endif
 }
diff --git a/Engine/Source/Runtime/Net/Core/Private/Net/Core/Misc/ResizableCircularQueueTests.cpp b/Engine/Source/Runtime/Net/Core/Private/Net/Core/Misc/ResizableCircularQueueTests.cpp
--- a/Engine/Source/Runtime/Net/Core/Private/Net/Core/Misc/ResizableCircularQueueTests.cpp
+++ b/Engine/Source/Runtime/Net/Core/Private/Net/Core/Misc/ResizableCircularQueueTests.cpp
@@ -11,12 +11,12 @@ struct FResizableCircularQueueTestUtil
 {
 	typedef TResizableCircularQueue<uint32> QueueT;
 
-	static bool VerifyQueueIntegrity(const QueueT& Queue, uint32 ExpectedValueAtFront, uint32 Increment)
+	static bool VerifyQueueIntegrity(const QueueT& Queue, QueueT::ElementT ExpectedValueAtFront, QueueT::ElementT Increment)
 	{
 		bool bSuccess = true;
 		SIZE_T Offset = Queue.Count() - 1;
 	
-		uint32 ExpectedValue = ExpectedValueAtFront;
+		QueueT::ElementT ExpectedValue = ExpectedValueAtFront;
 
 		// Peek elements in queue at given offset, peek from back to front
 		for (SIZE_T It = 0, EndIt = Queue.Count(); It < EndIt; ++It)
@@ -37,6 +37,8 @@ struct FResizableCircularQueueTestUtil
 
 bool FResizableCircularQueueTest::RunTest(const FString& Parameters)
 {
+	typedef FResizableCircularQueueTestUtil::QueueT::ElementT ElementT;
+
 	// Test empty
 	{
 		FResizableCircularQueueTestUtil::QueueT Q(0);
@@ -48,11 +50,11 @@ bool FResizableCircularQueueTest::RunTest(const FString& Parameters)
 
 	// Test Push to Capacity
 	{
-		const SIZE_T ElementsToPush = 8;
+		const ElementT ElementsToPush = 8;
 
 		FResizableCircularQueueTestUtil::QueueT Q(ElementsToPush);
 
-		for (SIZE_T It=0; It < ElementsToPush; ++It)
+		for (ElementT It=0; It < ElementsToPush; ++It)
 		{
 			Q.Enqueue(It);
 		}
@@ -64,11 +66,11 @@ bool FResizableCircularQueueTest::RunTest(const FString& Parameters)
 
 	// Test Push over Capacity
 	{
-		const SIZE_T InitialQueueCapacity = 32;
+		const ElementT InitialQueueCapacity = 32;
 
 		FResizableCircularQueueTestUtil::QueueT Q(InitialQueueCapacity);
 
-		for (SIZE_T It=0; It < InitialQueueCapacity; ++It)
+		for (ElementT It=0; It < InitialQueueCapacity; ++It)
 		{
 			Q.Enqueue(It);
 		}
@@ -83,7 +85,7 @@ bool FResizableCircularQueueTest::RunTest(const FString& Parameters)
 
 	// Test Push and Pop
 	{
-		const SIZE_T ElementsToPush = 256;
+		const ElementT ElementsToPush = 256;
 		const SIZE_T ElementPopMod = 16;
 		const SIZE_T ExpectedSize = 256 - ElementPopMod;
 		const SIZE_T ExpectedCapacity = 256;
@@ -91,7 +93,7 @@ bool FResizableCircularQueueTest::RunTest(const FString& Parameters)
 		FResizableCircularQueueTestUtil::QueueT Q(4);
 
 		uint32 ExpectedPoppedValue = 0;
-		for (uint32 It = 0; It < 256; ++It)
+		for (ElementT It = 0; It < ElementsToPush; ++It)
 		{
 			Q.Enqueue(It);
 			TestEqual(TEXT("Test Push and pop - Push"), It, Q.PeekAtOffset(Q.Count() - 1));
@@ -112,14 +114,14 @@ bool FResizableCircularQueueTest::RunTest(const FString& Parameters)
 
 	// Test Push and pop all
 	{
-		const SIZE_T ElementsToPush = 256;
+		const ElementT ElementsToPush = 256;
 
 		FResizableCircularQueueTestUtil::QueueT Q(ElementsToPush);
 
 		TestTrue(TEXT ("Test Push and pop all - IsEmpty before"), Q.IsEmpty());
 		TestEqual(TEXT("Test Push and pop all - Size before"), Q.Count(), SIZE_T(0));
 
-		for (SIZE_T It=0; It < ElementsToPush; ++It)
+		for (ElementT It=0; It < ElementsToPush; ++It)
 		{
 			Q.Enqueue(It);
 		}
@@ -128,7 +130,7 @@ bool FResizableCircularQueueTest::RunTest(const FString& Parameters)
 		TestEqual(TEXT("Test Push and pop all - Capacity"), Q.AllocatedCapacity(), ElementsToPush);
  		TestTrue (TEXT("Test Push and pop all - Expected"), FResizableCircularQueueTestUtil::VerifyQueueIntegrity(Q, 0, 1));
 
-		for (SIZE_T It=0; It < ElementsToPush; ++It)
+		for (ElementT It=0; It < ElementsToPush; ++It)
 		{
 			Q.Pop();
 		}
@@ -181,7 +183,7 @@ bool FResizableCircularQueueTest::RunTest(const FString& Parameters)
 
 	// Test index wrap
 	{
-		const SIZE_T ElementsToPush = 256;
+		const ElementT ElementsToPush = 256;
 		const SIZE_T ElementPopMod = 16;
 		const SIZE_T ExpectedSize = 256 - ElementPopMod;
 		const SIZE_T ExpectedCapacity = 256;
@@ -194,7 +196,7 @@ bool FResizableCircularQueueTest::RunTest(const FString& Parameters)
 		TestTrue(TEXT ("Test index wrap - IsEmpty before"), Q.IsEmpty());
 		TestEqual(TEXT("Test index wrap - Size before"), Q.Count(), SIZE_T(0));
 
-		for (SIZE_T It=0; It < ElementsToPush; ++It)
+		for (ElementT It=0; It < ElementsToPush; ++It)
 		{
 			Q.Enqueue(It);
 		}
@@ -203,7 +205,7 @@ bool FResizableCircularQueueTest::RunTest(const FString& Parameters)
 		TestEqual(TEXT("Test index wrap - Capacity"), Q.AllocatedCapacity(), ElementsToPush);
  		TestTrue (TEXT("Test index wrap - Expected"), FResizableCircularQueueTestUtil::VerifyQueueIntegrity(Q, 0, 1));
 
-		for (SIZE_T It=0; It < ElementsToPush; ++It)
+		for (ElementT It=0; It < ElementsToPush; ++It)
 		{
 			Q.Pop();
 		}
@@ -215,14 +217,14 @@ bool FResizableCircularQueueTest::RunTest(const FString& Parameters)
 
 	// Test Trim
 	{
-		const SIZE_T ElementsToPush = 9;
+		const ElementT ElementsToPush = 9;
 		const SIZE_T ElementsToPop = 5;
 		const SIZE_T ExpectedCapacity = 16;
 		const SIZE_T ExpectedCapacityAfterTrim = 4;
 
 		FResizableCircularQueueTestUtil::QueueT Q(0);
 
-		for (SIZE_T It=0; It < ElementsToPush; ++It)
+		for (ElementT It=0; It < ElementsToPush; ++It)
 		{
 			Q.Enqueue(It);
 		}
diff --git a/Engine/Source/Runtime/Net/Core/Private/Net/Core/PushModel/PushModel.cpp b/Engine/Source/Runtime/Net/Core/Private/Net/Core/PushModel/PushModel.cpp
--- a/Engine/Source/Runtime/Net/Core/Private/Net/Core/PushModel/PushModel.cpp
+++ b/Engine/Source/Runtime/Net/Core/Private/Net/Core/PushModel/PushModel.cpp
@@ -219,7 +219,7 @@ namespace UEPushModelPrivate
 			if (LIKELY(PerObjectStates.IsValidIndex(ObjectIndex)))
 			{
 				// The macros will take care of filtering out invalid objects, so we don't need to check here.
-				PerObjectStates[ObjectIndex].MarkPropertyDirty(RepIndex);
+				PerObjectStates[ObjectIndex].MarkPropertyDirty(static_cast<uint16>(RepIndex));
 			}
 		}
 
@@ -229,9 +229,9 @@ namespace UEPushModelPrivate
 			if (LIKELY(PerObjectStates.IsValidIndex(ObjectIndex)))
 			{
 				FPushModelPerObjectState& ObjectState = PerObjectStates[ObjectIndex];
-				for (int RepIndex = StartRepIndex; RepIndex <= EndRepIndex; ++RepIndex)
+				for (int32 RepIndex = StartRepIndex; RepIndex <= EndRepIndex; ++RepIndex)
 				{
-					ObjectState.MarkPropertyDirty(RepIndex);
+					ObjectState.MarkPropertyDirty(static_cast<uint16>(RepIndex));
 				}
 			}
 		}
diff --git a/Engine/Source/Runtime/Net/Core/Private/Net/Core/PushModel/Types/PushModelPerObjectState.h b/Engine/Source/Runtime/Net/Core/Private/Net/Core/PushModel/Types/PushModelPerObjectState.h
--- a/Engine/Source/Runtime/Net/Core/Private/Net/Core/PushModel/Types/PushModelPerObjectState.h
+++ b/Engine/Source/Runtime/Net/Core/Private/Net/Core/PushModel/Types/PushModelPerObjectState.h
@@ -98,7 +98,7 @@ namespace UEPushModelPrivate
 		FNetPushPerNetDriverId AddPerNetDriverState()
 		{
 			FSparseArrayAllocationInfo AllocationInfo = PerNetDriverStates.AddUninitialized();
-			new (AllocationInfo.Pointer) FPushModelPerNetDriverState(DirtiedThisFrame.Num());
+			new (AllocationInfo.Pointer) FPushModelPerNetDriverState(static_cast<uint16>(DirtiedThisFrame.Num()));
 			
 			return AllocationInfo.Index;
 		}
diff --git a/Engine/Source/Runtime/Net/Core/Private/Net/Core/Serialization/QuantizedVectorSerialization.cpp b/Engine/Source/Runtime/Net/Core/Private/Net/Core/Serialization/QuantizedVectorSerialization.cpp
--- a/Engine/Source/Runtime/Net/Core/Private/Net/Core/Serialization/QuantizedVectorSerialization.cpp
+++ b/Engine/Source/Runtime/Net/Core/Private/Net/Core/Serialization/QuantizedVectorSerialization.cpp
@@ -60,7 +60,7 @@ bool WriteQuantizedVector(const int32 Scale, const T& Value, FArchive& Ar)
 		return false;
 	}
 
-	const ScalarType Factor = Scale;
+	const ScalarType Factor = IntCastChecked<int16>(Scale);
 	T ScaledValue;
 	ScaledValue.X = Value.X*Factor;
 	ScaledValue.Y = Value.Y*Factor;
