Iris - Implemented FNetBitArrayConstRangeIterator. Can only be used in range-based for loops.

diff --git a/Engine/Plugins/Runtime/ReplicationSystemTestPlugin/Source/Private/Tests/Core/TestNetBitArray.cpp b/Engine/Plugins/Runtime/ReplicationSystemTestPlugin/Source/Private/Tests/Core/TestNetBitArray.cpp
--- a/Engine/Plugins/Runtime/ReplicationSystemTestPlugin/Source/Private/Tests/Core/TestNetBitArray.cpp
+++ b/Engine/Plugins/Runtime/ReplicationSystemTestPlugin/Source/Private/Tests/Core/TestNetBitArray.cpp
@@ -1628,4 +1628,55 @@ UE_NET_TEST_FIXTURE(FNetBitArrayFixture, MakeNetBitArrayView)
 	}
 }
 
+UE_NET_TEST(FNetBitArrayConstRangeIterator, CanIterateOverEmptyBitArray)
+{
+	bool bIteratedOverEmptyArray = false;
+	FNetBitArrayView Empty;
+	for (const uint32 Index : Empty)
+	{
+		bIteratedOverEmptyArray = true;
+	}
+
+	UE_NET_ASSERT_FALSE(bIteratedOverEmptyArray);
+}
+
+UE_NET_TEST(FNetBitArrayConstRangeIterator, CanIterateOverArbitrailySizedBitArrayWithAllBitsSet)
+{
+	FNetBitArrayBase::StorageWordType WordBuffer[] = { ~0U, ~0U, ~0U, ~0U, ~0U, ~0U, ~0U, };
+
+	constexpr uint32 ArraySize = 126U;
+	FNetBitArrayView BitArray(WordBuffer, ArraySize, FNetBitArrayView::NoResetNoValidate);
+	uint32 It = 0;
+	for (const uint32 Index : BitArray)
+	{
+		// Make sure we're hitting the right indices.
+		UE_NET_ASSERT_EQ(Index, It);
+		++It;
+	}
+
+	UE_NET_ASSERT_EQ(It, ArraySize);
+}
+
+UE_NET_TEST(FNetBitArrayConstRangeIterator, CanIterateOverArbitrailySizedBitArrayWithArbitraryBitsSet)
+{
+	constexpr uint32 ArraySize = 128U;
+	FNetBitArray BitArray(ArraySize);
+
+	const uint32 IndicesToSet[] = { 65, 68 };
+	for (uint32 IndexToSet : IndicesToSet)
+	{
+		BitArray.SetBit(IndexToSet);
+	}
+
+	uint32 It = 0;
+	for (const uint32 Index : BitArray)
+	{
+		// Make sure we're hitting the right indices.
+		UE_NET_ASSERT_EQ(Index, IndicesToSet[It]);
+		++It;
+	}
+
+	UE_NET_ASSERT_EQ(It, (uint32)UE_ARRAY_COUNT(IndicesToSet));
+}
+
 }
diff --git a/Engine/Source/Runtime/Net/Core/Public/Net/Core/NetBitArray.h b/Engine/Source/Runtime/Net/Core/Public/Net/Core/NetBitArray.h
--- a/Engine/Source/Runtime/Net/Core/Public/Net/Core/NetBitArray.h
+++ b/Engine/Source/Runtime/Net/Core/Public/Net/Core/NetBitArray.h
@@ -7,9 +7,19 @@
 #include "HAL/PlatformMath.h"
 #include "HAL/PlatformMemory.h"
 #include "Misc/AssertionMacros.h"
+#include <type_traits>
+
+namespace UE::Net
+{
+	class FNetBitArray;
+	class FNetBitArrayView;
+
+	namespace Private
+	{
+		class FNetBitArrayRangedForConstIterator;
+	}
+}
 
-namespace UE { namespace Net { class FNetBitArray; } }
-namespace UE { namespace Net { class FNetBitArrayView; } }
 
 /* NetBitArray validation support. */
 #ifndef UE_NETBITARRAY_VALIDATE
@@ -194,6 +204,13 @@ public:
 	template<typename T, typename V>
 	static void ForAllExclusiveBits(const FNetBitArray& A, const FNetBitArray& B, T&& FunctorA, V&& FunctorB);
 
+	/**
+	 * DO NOT USE DIRECTLY
+	 * Range-based for loop support iterating over set bits and returning the index to the set bit.
+	 */
+	Private::FNetBitArrayRangedForConstIterator begin() const;
+	Private::FNetBitArrayRangedForConstIterator end() const;
+
 private:
 	StorageWordType GetLastWordMask() const { return (~StorageWordType(0) >> (uint32(-int32(BitCount)) & (WordBitCount - 1))); }
 	void ClearPaddingBits();
@@ -335,6 +352,13 @@ public:
 
 	static constexpr inline uint32 CalculateRequiredWordCount(uint32 BitCount);
 
+	/**
+	 * DO NOT USE DIRECTLY
+	 * Range-based for loop support iterating over set bits and returning the index to the set bit.
+	 */
+	Private::FNetBitArrayRangedForConstIterator begin() const;
+	Private::FNetBitArrayRangedForConstIterator end() const;
+
 private:
 	inline StorageWordType GetLastWordMask() const { return (~StorageWordType(0) >> (uint32(-int32(BitCount)) & (WordBitCount - 1))); }
 
@@ -343,7 +367,6 @@ private:
 	uint32 BitCount;
 };
 
-
 // Implementation helper
 class FNetBitArrayHelper final
 {
@@ -707,6 +730,43 @@ private:
 
 };
 
+namespace Private
+{
+
+class FNetBitArrayRangedForConstIterator
+{
+public:
+	FNetBitArrayRangedForConstIterator();
+
+	FNetBitArrayRangedForConstIterator& operator++();
+
+	bool operator!=(const FNetBitArrayRangedForConstIterator& It) const;
+
+	/** Returns the bit index the iterator is pointing at. */
+	uint32 operator*() const;
+
+private:
+	friend FNetBitArray;
+	friend FNetBitArrayView;
+
+	enum class ERangeStart : unsigned
+	{
+		Begin,
+		End
+	};
+
+	FNetBitArrayRangedForConstIterator(const FNetBitArrayBase::StorageWordType* InData UE_LIFETIMEBOUND, uint32 BitCount, ERangeStart RangeStart);
+
+	void AdvanceToNextSetBit();
+
+	const FNetBitArrayBase::StorageWordType* Data = nullptr;
+	uint32 BitCount = 0;
+	uint32 WordCount = 0;
+	uint32 CurrentBitIndex = 0;
+	uint32 CurrentWord = 0;
+};
+
+}
 
 // NetBitArray Implementation
 inline FNetBitArray::FNetBitArray()
@@ -905,6 +965,17 @@ inline void FNetBitArray::ClearPaddingBits()
 	}
 }
 
+inline Private::FNetBitArrayRangedForConstIterator FNetBitArray::begin() const
+{
+	return Private::FNetBitArrayRangedForConstIterator(Storage.GetData(), BitCount, Private::FNetBitArrayRangedForConstIterator::ERangeStart::Begin);
+}
+
+inline Private::FNetBitArrayRangedForConstIterator FNetBitArray::end() const
+{
+	return Private::FNetBitArrayRangedForConstIterator(Storage.GetData(), BitCount, Private::FNetBitArrayRangedForConstIterator::ERangeStart::End);
+}
+
+
 // NetBitArrayView implementation
 constexpr uint32 FNetBitArrayView::CalculateRequiredWordCount(uint32 BitCount)
 { 
@@ -1098,6 +1169,16 @@ void FNetBitArrayView::ForAllExclusiveBits(const FNetBitArrayView& A, const FNet
 	FNetBitArrayHelper::ForAllExclusiveBits(A.Storage, B.Storage, A.WordCount, A.BitCount, FunctorA, FunctorB);
 }
 
+inline Private::FNetBitArrayRangedForConstIterator FNetBitArrayView::begin() const
+{
+	return Private::FNetBitArrayRangedForConstIterator(Storage, BitCount, Private::FNetBitArrayRangedForConstIterator::ERangeStart::Begin);
+}
+
+inline Private::FNetBitArrayRangedForConstIterator FNetBitArrayView::end() const
+{
+	return Private::FNetBitArrayRangedForConstIterator(Storage, BitCount, Private::FNetBitArrayRangedForConstIterator::ERangeStart::End);
+}
+
 inline FNetBitArrayView MakeNetBitArrayView(const FNetBitArrayView::StorageWordType* Storage, uint32 BitCount)
 {
 	return FNetBitArrayView(const_cast<FNetBitArrayView::StorageWordType*>(Storage), BitCount);
@@ -1109,3 +1190,67 @@ inline FNetBitArrayView MakeNetBitArrayView(const FNetBitArray& BitArray)
 }
 
 } // end namespace UE::Net
+
+namespace UE::Net::Private
+{
+
+inline FNetBitArrayRangedForConstIterator::FNetBitArrayRangedForConstIterator()
+{
+}
+
+inline FNetBitArrayRangedForConstIterator::FNetBitArrayRangedForConstIterator(const FNetBitArrayBase::StorageWordType* InData UE_LIFETIMEBOUND, uint32 InBitCount, FNetBitArrayRangedForConstIterator::ERangeStart RangeStart)
+: Data(InData)
+, BitCount(InBitCount)
+, WordCount(FNetBitArrayView::CalculateRequiredWordCount(BitCount))
+, CurrentBitIndex(RangeStart == ERangeStart::Begin ? 0 : BitCount)
+{
+	UE_NETBITARRAY_CHECK(InData != nullptr || InBitCount == 0);
+	CurrentWord = (CurrentBitIndex < BitCount ? Data[0] : FNetBitArrayBase::StorageWordType(0));
+	AdvanceToNextSetBit();
+}
+
+inline FNetBitArrayRangedForConstIterator& FNetBitArrayRangedForConstIterator::operator++()
+{
+	AdvanceToNextSetBit();
+	return *this;
+}
+
+inline uint32 FNetBitArrayRangedForConstIterator::operator*() const
+{
+	return CurrentBitIndex;
+}
+
+inline bool FNetBitArrayRangedForConstIterator::operator!=(const FNetBitArrayRangedForConstIterator& It) const
+{
+	UE_NETBITARRAY_CHECK(this->Data == It.Data);
+	return this->CurrentBitIndex != It.CurrentBitIndex;
+}
+
+inline void FNetBitArrayRangedForConstIterator::AdvanceToNextSetBit()
+{
+	using SignedWordType = typename std::make_signed<FNetBitArrayBase::StorageWordType>::type;
+
+	uint32 WordIt = CurrentBitIndex/FNetBitArrayBase::WordBitCount;
+	while (CurrentWord == 0U)
+	{
+		++WordIt;
+		if (WordIt >= WordCount)
+		{
+			CurrentBitIndex = BitCount;
+			return;
+		}
+
+		CurrentWord = Data[WordIt];
+		// It's ok to not adjust the bit index to start on a word boundary. It's handled outside the loop.
+		CurrentBitIndex += FNetBitArrayBase::WordBitCount;
+	}
+
+	const uint32 NewBitIndex = (CurrentBitIndex & ~(FNetBitArrayBase::WordBitCount - 1U)) + FPlatformMath::CountTrailingZeros(CurrentWord);
+	CurrentBitIndex = FPlatformMath::Min(NewBitIndex, BitCount);
+
+	// Clear least significant bit from CurrentWord. It doesn't matter if we're at the end if the bit array as the CurrentWord will be ignored in that case.
+	const FNetBitArrayBase::StorageWordType LeastSignificantBit = CurrentWord & FNetBitArrayBase::StorageWordType(-SignedWordType(CurrentWord));
+	CurrentWord ^= LeastSignificantBit;
+}
+
+}
