NetBitArray:

* Added IsBitSet
* Added CombineMultiple to run bit operations on two const bitarrays in 1 loop
* Added ToString that prints how many set bits the array currently has. Slightly expensive do not abuse with active logs.
* Added FNetBitArrayPrinter utility functions

diff --git a/Engine/Plugins/Runtime/ReplicationSystemTestPlugin/Source/Private/Tests/Core/TestNetBitArray.cpp b/Engine/Plugins/Runtime/ReplicationSystemTestPlugin/Source/Private/Tests/Core/TestNetBitArray.cpp
--- a/Engine/Plugins/Runtime/ReplicationSystemTestPlugin/Source/Private/Tests/Core/TestNetBitArray.cpp
+++ b/Engine/Plugins/Runtime/ReplicationSystemTestPlugin/Source/Private/Tests/Core/TestNetBitArray.cpp
@@ -3,6 +3,7 @@
 #include "NetworkAutomationTest.h"
 #include "NetworkAutomationTestMacros.h"
 #include "Net/Core/NetBitArray.h"
+#include "Net/Core/NetBitArrayPrinter.h"
 
 namespace UE::Net::Private
 {
@@ -1044,6 +1045,25 @@ UE_NET_TEST(FNetBitArrayView, Or)
 			UE_NET_ASSERT_EQ(ExpectedWordBuffer[it], WordBufferA[it]);
 		}
 	}
+
+	// Test Multiple arrays
+	{
+		const uint32 ExpectedWordBuffer[] = { 0x0011, 0x0110, 0x1100, 0x1001 };
+		uint32 WordBufferA[] = { 0x0000, 0x0000, 0x0000, 0x0000 };
+		uint32 WordBufferB[] = { 0x0001, 0x0010, 0x0100, 0x1000 };
+		uint32 WordBufferC[] = { 0x0010, 0x0100, 0x1000, 0x0001 };
+
+		FNetBitArrayView BitArrayA(&WordBufferA[0], 128);
+		const FNetBitArrayView BitArrayB(&WordBufferB[0], 128);
+		const FNetBitArrayView BitArrayC(&WordBufferC[0], 128);
+
+		BitArrayA.CombineMultiple(WordOp, BitArrayB, WordOp, BitArrayC);
+
+		for (uint32 it = 0; it < UE_ARRAY_COUNT(ExpectedWordBuffer); ++it)
+		{
+			UE_NET_ASSERT_EQ(ExpectedWordBuffer[it], WordBufferA[it]);
+		}
+	}
 }
 
 UE_NET_TEST(FNetBitArrayView, And)
@@ -1094,6 +1114,25 @@ UE_NET_TEST(FNetBitArrayView, And)
 			UE_NET_ASSERT_EQ(ExpectedWordBuffer[it], WordBufferA[it]);
 		}
 	}
+
+	// Test Multiple arrays
+	{
+		const uint32 ExpectedWordBuffer[] = { 0x0001, 0x0010, 0x0100, 0x1000 };
+		uint32 WordBufferA[] = { 0x0001, 0x0010, 0x0100, 0x1000 };
+		uint32 WordBufferB[] = { 0x0001, 0x0010, 0x0100, 0x1000 };
+		uint32 WordBufferC[] = { 0x0001, 0x0010, 0x0100, 0x1000 };
+
+		FNetBitArrayView BitArrayA(&WordBufferA[0], 128);
+		const FNetBitArrayView BitArrayB(&WordBufferB[0], 128);
+		const FNetBitArrayView BitArrayC(&WordBufferC[0], 128);
+
+		BitArrayA.CombineMultiple(WordOp, BitArrayB, WordOp, BitArrayC);
+
+		for (uint32 it = 0; it < UE_ARRAY_COUNT(ExpectedWordBuffer); ++it)
+		{
+			UE_NET_ASSERT_EQ(ExpectedWordBuffer[it], WordBufferA[it]);
+		}
+	}
 }
 
 UE_NET_TEST(FNetBitArrayView, AndNot)
@@ -1144,6 +1183,25 @@ UE_NET_TEST(FNetBitArrayView, AndNot)
 			UE_NET_ASSERT_EQ(ExpectedWordBuffer[it], WordBufferA[it]);
 		}
 	}
+
+	// Test Multiple arrays
+	{
+		const uint32 ExpectedWordBuffer[] = { 0xFFFe, 0xFFeF, 0xFeFF, 0xeFFF };
+		uint32 WordBufferA[] = { 0xFFFF, 0xFFFF, 0xFFFF, 0xFFFF };
+		uint32 WordBufferB[] = { 0x0001, 0x0010, 0x0100, 0x1000 };
+		uint32 WordBufferC[] = { 0x0000, 0x0000, 0x0000, 0x0000 };
+
+		FNetBitArrayView BitArrayA(&WordBufferA[0], 128);
+		const FNetBitArrayView BitArrayB(&WordBufferB[0], 128);
+		const FNetBitArrayView BitArrayC(&WordBufferC[0], 128);
+
+		BitArrayA.CombineMultiple(WordOp, BitArrayB, WordOp, BitArrayC);
+
+		for (uint32 it = 0; it < UE_ARRAY_COUNT(ExpectedWordBuffer); ++it)
+		{
+			UE_NET_ASSERT_EQ(ExpectedWordBuffer[it], WordBufferA[it]);
+		}
+	}
 }
 
 UE_NET_TEST(FNetBitArrayView, Xor)
@@ -1194,6 +1252,50 @@ UE_NET_TEST(FNetBitArrayView, Xor)
 			UE_NET_ASSERT_EQ(ExpectedWordBuffer[it], WordBufferA[it]);
 		}
 	}
+
+	// Test Multiple arrays
+	{
+		const uint32 ExpectedWordBuffer[] = { 0xFFFF, 0xFFFF, 0xFFFF, 0xFFFF };
+		uint32 WordBufferA[] = { 0xFFee, 0xFeeF, 0xeeFF, 0xeFFe };
+		uint32 WordBufferB[] = { 0x0001, 0x0010, 0x0100, 0x1000 };
+		uint32 WordBufferC[] = { 0x0010, 0x0100, 0x1000, 0x0001 };
+
+		FNetBitArrayView BitArrayA(&WordBufferA[0], 128);
+		const FNetBitArrayView BitArrayB(&WordBufferB[0], 128);
+		const FNetBitArrayView BitArrayC(&WordBufferC[0], 128);
+
+		BitArrayA.CombineMultiple(WordOp, BitArrayB, WordOp, BitArrayC);
+
+		for (uint32 it = 0; it < UE_ARRAY_COUNT(ExpectedWordBuffer); ++it)
+		{
+			UE_NET_ASSERT_EQ(ExpectedWordBuffer[it], WordBufferA[it]);
+		}
+	}
+}
+
+UE_NET_TEST(FNetBitArrayView, TestCombineMultiple)
+{
+	auto&& FirstWordOp = FNetBitArrayView::OrOp;
+	auto&& SecondWordOp = FNetBitArrayView::AndOp;
+
+	// Test Multiple arrays
+	{
+		const uint32 ExpectedWordBuffer[] = { 0xFF10, 0xF10F, 0x10FF, 0x0FF1 };
+		uint32 WordBufferA[] = { 0xFF00, 0xF00F, 0x00FF, 0x0FF0 };
+		uint32 WordBufferB[] = { 0x0011, 0x0110, 0x1100, 0x1001 };
+		uint32 WordBufferC[] = { 0x0010, 0x0100, 0x1000, 0x0001 };
+
+		FNetBitArrayView BitArrayA(&WordBufferA[0], 128);
+		const FNetBitArrayView BitArrayB(&WordBufferB[0], 128);
+		const FNetBitArrayView BitArrayC(&WordBufferC[0], 128);
+
+		BitArrayA.CombineMultiple(FirstWordOp, BitArrayB, SecondWordOp, BitArrayC);
+
+		for (uint32 it = 0; it < UE_ARRAY_COUNT(ExpectedWordBuffer); ++it)
+		{
+			UE_NET_ASSERT_EQ(ExpectedWordBuffer[it], WordBufferA[it]);
+		}
+	}
 }
 
 UE_NET_TEST(FNetBitArrayView, Copy)
@@ -1231,6 +1333,55 @@ UE_NET_TEST(FNetBitArrayView, Copy)
 	}
 }
 
+UE_NET_TEST(FNetBitArrayView, PrintHelper)
+{
+	{
+		uint32 WordBufferA[] = { 0xFF00, 0xF00F, 0x00FF, 0x0FF0 };
+
+		FNetBitArray BitArrayB;
+		BitArrayB.Init(128);
+		BitArrayB.GetData()[0] = WordBufferA[0];
+		BitArrayB.GetData()[1] = WordBufferA[1];
+		BitArrayB.GetData()[2] = WordBufferA[2];
+		BitArrayB.GetData()[3] = WordBufferA[3];
+		
+		const FNetBitArrayView BitArrayA(&WordBufferA[0], 128);
+		FString PrintBufferA;
+		FString PrintBufferB;
+
+		PrintBufferA = FNetBitArrayPrinter::PrintSetBits(BitArrayA);
+		PrintBufferB = FNetBitArrayPrinter::PrintSetBits(BitArrayB);
+		UE_NET_ASSERT_EQ(PrintBufferA, PrintBufferB);
+
+		PrintBufferA = FNetBitArrayPrinter::PrintSetSummary(BitArrayA);
+		PrintBufferB = FNetBitArrayPrinter::PrintSetSummary(BitArrayB);
+		UE_NET_ASSERT_EQ(PrintBufferA, PrintBufferB);
+
+		PrintBufferA = BitArrayA.ToString();
+		PrintBufferB = BitArrayB.ToString();
+		UE_NET_ASSERT_EQ(PrintBufferA, PrintBufferB);
+
+		PrintBufferA = FNetBitArrayPrinter::PrintZeroBits(BitArrayA);
+		PrintBufferB = FNetBitArrayPrinter::PrintZeroBits(BitArrayB);
+		UE_NET_ASSERT_EQ(PrintBufferA, PrintBufferB);
+		
+		PrintBufferA = FNetBitArrayPrinter::PrintZeroSummary(BitArrayA);
+		PrintBufferB = FNetBitArrayPrinter::PrintZeroSummary(BitArrayB);
+		UE_NET_ASSERT_EQ(PrintBufferA, PrintBufferB);
+
+		uint32 WordBufferDelta[] = { 0x0011, 0xF00F, 0x1100, 0x0FF0 };
+		const FNetBitArrayView BitArrayDelta(&WordBufferDelta[0], 128);
+
+		PrintBufferA = FNetBitArrayPrinter::PrintDeltaSummary(BitArrayA, BitArrayDelta);
+		PrintBufferB = FNetBitArrayPrinter::PrintDeltaSummary(BitArrayDelta, BitArrayB);
+		UE_NET_ASSERT_EQ(PrintBufferA, PrintBufferB);
+
+		PrintBufferA = FNetBitArrayPrinter::PrintDeltaBits(BitArrayA, BitArrayDelta);
+		PrintBufferB = FNetBitArrayPrinter::PrintDeltaBits(BitArrayB, BitArrayDelta);
+		UE_NET_ASSERT_EQ(PrintBufferA, PrintBufferB);
+	}
+}
+
 class FNetBitArrayViewFixture : public FNetworkAutomationTestSuiteFixture
 {
 public:
diff --git a/Engine/Source/Runtime/Net/Core/Private/Net/Core/NetBitArray.cpp b/Engine/Source/Runtime/Net/Core/Private/Net/Core/NetBitArray.cpp
--- a/Engine/Source/Runtime/Net/Core/Private/Net/Core/NetBitArray.cpp
+++ b/Engine/Source/Runtime/Net/Core/Private/Net/Core/NetBitArray.cpp
@@ -1,11 +1,34 @@
 // Copyright Epic Games, Inc. All Rights Reserved.
 
 #include "Net/Core/NetBitArray.h"
+
+#include "Net/Core/NetBitArrayPrinter.h"
+
 #include "Traits/IntType.h"
 
 namespace UE::Net
 {
 
+//*************************************************************************************************
+// FNetBitArray
+//*************************************************************************************************
+FString FNetBitArray::ToString() const
+{
+	return FNetBitArrayPrinter::PrintSetSummary(*this);
+}
+
+//*************************************************************************************************
+// FNetBitArrayView
+//*************************************************************************************************
+FString FNetBitArrayView::ToString() const
+{
+	return FNetBitArrayPrinter::PrintSetSummary(*this);
+}
+
+//*************************************************************************************************
+// FNetBitArrayHelper
+//*************************************************************************************************
+
 uint32 FNetBitArrayHelper::GetSetBitIndices(const FNetBitArrayHelper::StorageWordType* Storage, const uint32 BitCount, const uint32 StartIndex, const uint32 Count, uint32* const OutIndices, const uint32 OutIndicesCapacity)
 {
 	if ((!ensure(OutIndicesCapacity > 0)) | (StartIndex >= BitCount) | (Count == 0))
@@ -68,4 +91,4 @@ uint32 FNetBitArrayHelper::CountSetBits(const StorageWordType* Storage, const ui
 	return SetBitsCount;
 }
 
-}
+} // end namespace UE::Net
\ No newline at end of file
diff --git a/Engine/Source/Runtime/Net/Core/Public/Net/Core/NetBitArray.h b/Engine/Source/Runtime/Net/Core/Public/Net/Core/NetBitArray.h
--- a/Engine/Source/Runtime/Net/Core/Public/Net/Core/NetBitArray.h
+++ b/Engine/Source/Runtime/Net/Core/Public/Net/Core/NetBitArray.h
@@ -37,6 +37,7 @@ namespace UE::Net
 #endif
 
 #define UE_NETBITARRAY_VALIDATE_COMPATIBLE(x) UE_NETBITARRAY_CHECK(this != &(x) && BitCount == (x).BitCount)
+#define UE_NETBITARRAY_VALIDATE_BOTH_COMPATIBLE(lhs,rhs) UE_NETBITARRAY_CHECK((void*)(&lhs) != (void*)(&rhs) && (lhs).GetNumBits() == (rhs).GetNumBits())
 #define UE_NETBITARRAYVIEW_VALIDATE_COMPATIBLE(x) UE_NETBITARRAY_CHECK((Storage != (x).Storage) && (BitCount == (x).BitCount))
 
 namespace UE::Net
@@ -119,6 +120,9 @@ public:
 	/** Sets all bits in the array, but clears padding bits. */
 	void SetAllBits();
 
+	/** Return if a specified bit is set or not */
+	bool IsBitSet(uint32 Index) const { return GetBit(Index); }
+
 	/** Returns true if any bit is set in the bitset. Note: Padding bits in storage are expected to be zero. */
 	bool IsAnyBitSet() const;
 
@@ -164,6 +168,9 @@ public:
 	/** Find last set bit. Returns InvalidIndex if no set bit was found. */
 	uint32 FindLastOne() const;
 
+	/** Only prints the amount of set bits in the array. @see FNetBitArrayPrinter for more print options */
+	NETCORE_API FString ToString() const;
+
 	/**
 	 * Retrieves set bits in the provided range and returns how many indices were written to OutIndices.
 	 * If OutIndices is filled to its capacity the search for set bits will end. OutIndices may be modified
@@ -192,6 +199,18 @@ public:
 	 */
 	template<typename WordOpFunctor> void Combine(const FNetBitArray& Other, WordOpFunctor&& WordOp);
 
+	/**
+	 * Combine this array with the result of a bitwise operation of two other arrays.
+	 * eg: ThisArray = ThisArray Op (ArrayA Op2 ArrayB)
+	 *
+	 * @param Op Operation to execute on the word of this array and the result of the operation on the two other arrays
+	 * @param ArrayA The first bit array to execute the other operation on
+	 * @param Op2 The operation to execute on the word of the read-only array passed in parameter
+	 * @param ArrayB The second bit array to execute the other operation on
+	 * @see FNetBitArrayHelper::AndOp, FNetBitArrayHelper::AndNotOp, FNetBitArrayHelper::OrOp, FNetBitArrayHelper::XorOp
+	 */
+	template<typename WordOpFunctor> void CombineMultiple(WordOpFunctor&& Op, const FNetBitArray& ArrayA, WordOpFunctor&& Op2, const FNetBitArray& ArrayB);
+
 	/** Iterate over all set bits and invoke functor with signature void(uint32 BitIndex). */
 	template<typename T>
 	void ForAllSetBits(T&& Functor) const;
@@ -244,6 +263,9 @@ public:
 	/** Return true if equal including BitCount and padding bits */
 	bool operator==(const FNetBitArrayView& Other) const;
 
+	/** Return if a specified bit is set or not */
+	inline bool IsBitSet(uint32 Index) const { return GetBit(Index); }
+
 	/** Returns true if any bit is set in the bitset Note: Padding bits in storage are expected to be zero. */
 	inline bool IsAnyBitSet() const;
 
@@ -310,6 +332,9 @@ public:
 	/** Returns a pointer to the internal storage. */
 	StorageWordType* GetData() { return Storage; }
 
+	/** Only prints the amount of set bits in the array. @see FNetBitArrayPrinter for more print options */
+	NETCORE_API FString ToString() const;
+
 	/**
 	 * Retrieves set bits in the provided range and returns how many indices were written to OutIndices.
 	 * If OutIndices is filled to its capacity the search for set bits will end. OutIndices may be modified
@@ -336,7 +361,19 @@ public:
 	 * @param WordOpFunctor The operation to perform for each word to combine. Signature is StorageWordType(StorageWordType ThisWord, StorageWordType OtherWord).
 	 * @see FNetBitArrayHelper::AndOp, FNetBitArrayHelper::AndNotOp, FNetBitArrayHelper::OrOp, FNetBitArrayHelper::XorOp
 	 */
-	template<typename Functor> void Combine(const FNetBitArrayView& Other, Functor&& WordOp);
+	template<typename WordOpFunctor> void Combine(const FNetBitArrayView& Other, WordOpFunctor&& WordOp);
+	
+	/**
+	 * Combine this array with the result of a bitwise operation of two other arrays.
+	 * eg: ThisArray = ThisArray Op (ArrayA Op2 ArrayB)
+	 *
+	 * @param Op Operation to execute on the word of this array and the result of the operation on the two other arrays
+	 * @param ArrayA The first bit array to execute the other operation on
+	 * @param Op2 The operation to execute on the word of the read-only array passed in parameter
+	 * @param ArrayB The second bit array to execute the other operation on
+	 * @see FNetBitArrayHelper::AndOp, FNetBitArrayHelper::AndNotOp, FNetBitArrayHelper::OrOp, FNetBitArrayHelper::XorOp
+	 */
+	template<typename WordOpFunctor> void CombineMultiple(WordOpFunctor&& Op, const FNetBitArrayView& ArrayA, WordOpFunctor&& Op2, const FNetBitArrayView& ArrayB);
 
 	/** Iterate over all set bits and invoke functor with signature void(uint32 BitIndex). */
 	template <typename T>
@@ -367,7 +404,11 @@ private:
 	uint32 BitCount;
 };
 
+//*************************************************************************************************
+// FNetBitArrayHelper 
 // Implementation helper
+//*************************************************************************************************
+
 class FNetBitArrayHelper final
 {
 private:
@@ -634,6 +675,14 @@ private:
 		}
 	}
 
+	template<typename Functor> static void CombineMultiple(StorageWordType* WriteStorage, const StorageWordType* const ReadStorageA, const StorageWordType* const ReadStorageB, const uint32 WordCount, Functor&& WordOpWrite, Functor&& WordOpRead)
+	{
+		for (uint32 WordIt = 0; WordIt < WordCount; ++WordIt)
+		{
+			WriteStorage[WordIt] = WordOpWrite(WriteStorage[WordIt], WordOpRead(ReadStorageA[WordIt], ReadStorageB[WordIt]));
+		}
+	}
+
 	static bool IsEqual(const StorageWordType* Storage, const StorageWordType* OtherStorage, const uint32 WordCount)
 	{
 		StorageWordType Result = 0U;
@@ -733,6 +782,9 @@ private:
 namespace Private
 {
 
+//*************************************************************************************************
+// FNetBitArrayRangedForConstIterator
+//*************************************************************************************************
 class FNetBitArrayRangedForConstIterator
 {
 public:
@@ -766,9 +818,11 @@ private:
 	uint32 CurrentWord = 0;
 };
 
-}
+} // end namespace Private
 
-// NetBitArray Implementation
+//*************************************************************************************************
+// FNetBitArray Implementation
+//*************************************************************************************************
 inline FNetBitArray::FNetBitArray()
 : BitCount(0)
 {
@@ -933,12 +987,19 @@ inline void FNetBitArray::Copy(const FNetBitArray& Other)
 	FPlatformMemory::Memcpy(GetData(), Other.GetData(), GetNumWords()*sizeof(StorageWordType));
 }
 
-template<typename Functor> inline void FNetBitArray::Combine(const FNetBitArray& Other, Functor&& WordOp)
+template<typename WordOpFunctor> inline void FNetBitArray::Combine(const FNetBitArray& Other, WordOpFunctor&& WordOp)
 {
 	UE_NETBITARRAY_VALIDATE_COMPATIBLE(Other);
 	FNetBitArrayHelper::Combine(Storage.GetData(), Other.Storage.GetData(), Storage.Num(), WordOp);
 }
 
+template<typename WordOpFunctor> inline void FNetBitArray::CombineMultiple(WordOpFunctor&& Op, const FNetBitArray& ArrayA, WordOpFunctor&& Op2, const FNetBitArray& ArrayB)
+{
+	UE_NETBITARRAY_VALIDATE_COMPATIBLE(ArrayA);
+	UE_NETBITARRAY_VALIDATE_COMPATIBLE(ArrayB);
+	FNetBitArrayHelper::CombineMultiple(Storage.GetData(), ArrayA.Storage.GetData(), ArrayB.Storage.GetData(), Storage.Num(), Op, Op2);
+}
+
 template<typename T> inline void FNetBitArray::ForAllSetBits(T&& Functor) const
 {
 	FNetBitArrayHelper::ForAllSetBits(Storage.GetData(), Storage.Num(), BitCount, Functor);
@@ -965,6 +1026,9 @@ inline void FNetBitArray::ClearPaddingBits()
 	}
 }
 
+//*************************************************************************************************
+// FNetBitArrayRangedForConstIterator implementation
+//*************************************************************************************************
 inline Private::FNetBitArrayRangedForConstIterator FNetBitArray::begin() const
 {
 	return Private::FNetBitArrayRangedForConstIterator(Storage.GetData(), BitCount, Private::FNetBitArrayRangedForConstIterator::ERangeStart::Begin);
@@ -976,7 +1040,9 @@ inline Private::FNetBitArrayRangedForConstIterator FNetBitArray::end() const
 }
 
 
+//*************************************************************************************************
 // NetBitArrayView implementation
+//*************************************************************************************************
 constexpr uint32 FNetBitArrayView::CalculateRequiredWordCount(uint32 BitCount)
 { 
 	return (BitCount + WordBitCount - 1) / WordBitCount; 
@@ -1144,12 +1210,19 @@ inline uint32 FNetBitArrayView::CountSetBits(uint32 StartIndex, uint32 Count) co
 	return FNetBitArrayHelper::CountSetBits(Storage, BitCount, StartIndex, Count);
 }
 
-template<typename Functor> inline void FNetBitArrayView::Combine(const FNetBitArrayView& Other, Functor&& WordOp)
+template<typename WordOpFunctor> inline void FNetBitArrayView::Combine(const FNetBitArrayView& Other, WordOpFunctor&& WordOp)
 {
 	UE_NETBITARRAYVIEW_VALIDATE_COMPATIBLE(Other);
 	FNetBitArrayHelper::Combine(Storage, Other.Storage, WordCount, WordOp);
 }
 
+template<typename WordOpFunctor> inline void FNetBitArrayView::CombineMultiple(WordOpFunctor&& Op, const FNetBitArrayView& ArrayA, WordOpFunctor&& Op2, const FNetBitArrayView& ArrayB)
+{
+	UE_NETBITARRAYVIEW_VALIDATE_COMPATIBLE(ArrayA);
+	UE_NETBITARRAYVIEW_VALIDATE_COMPATIBLE(ArrayB);
+	FNetBitArrayHelper::CombineMultiple(Storage, ArrayA.Storage, ArrayB.Storage, WordCount, Op, Op2);
+}
+
 template<typename T> inline void FNetBitArrayView::ForAllSetBits(T&& Functor) const
 {
 	FNetBitArrayHelper::ForAllSetBits(Storage, WordCount, BitCount, Functor);
@@ -1194,6 +1267,12 @@ inline FNetBitArrayView MakeNetBitArrayView(const FNetBitArray& BitArray)
 namespace UE::Net::Private
 {
 
+
+//*************************************************************************************************
+// FNetBitArrayRangedForConstIterator 
+// Implementation
+//*************************************************************************************************
+
 inline FNetBitArrayRangedForConstIterator::FNetBitArrayRangedForConstIterator()
 {
 }
@@ -1253,4 +1332,4 @@ inline void FNetBitArrayRangedForConstIterator::AdvanceToNextSetBit()
 	CurrentWord ^= LeastSignificantBit;
 }
 
-}
+} // end namespace UE::Net::Private
diff --git a/Engine/Source/Runtime/Net/Core/Public/Net/Core/NetBitArrayPrinter.h b/Engine/Source/Runtime/Net/Core/Public/Net/Core/NetBitArrayPrinter.h
new file mode 100644
--- /dev/null
+++ b/Engine/Source/Runtime/Net/Core/Public/Net/Core/NetBitArrayPrinter.h
@@ -0,0 +1,203 @@
+// Copyright Epic Games, Inc. All Rights Reserved.
+
+#pragma once
+
+#include "Net/Core/NetBitArray.h"
+
+namespace UE::Net
+{
+
+/**
+ * FNetBitArrayPrinter utility class
+ * Collection of helper functions that can describe any given BitArray into a string for logging or debugging purposes
+ */
+class FNetBitArrayPrinter
+{
+public:
+
+	/** Print the number of set bits in the bitarray */
+	template <class T>
+	static FString PrintSetSummary(const T& BitArray);
+
+	/** Print the number of zero bits int he bitarray */
+	template <class T>
+	static FString PrintZeroSummary(const T& BitArray);
+
+	/** Print the number of set bits in the bitarray and the value of each word with set bits in them */
+	template <class T>
+	static FString PrintSetBits(const T& BitArray);
+
+	/** Print the number of zero bits in the bitarray and the value of each word with zero bits in them */
+	template <class T>
+	static FString PrintZeroBits(const T& BitArray);
+
+	/** Print the number of bits different between each array */
+	template <class T, class U>
+	static FString PrintDeltaSummary(const T& BitArrayA, const U& BitArrayB);
+
+	/** Print the number of bits different between each array and the value of each word with different bits in them*/
+	template <class T, class U>
+	static FString PrintDeltaBits(const T& BitArrayA, const U& BitArrayB);
+
+private:
+
+	FNetBitArrayPrinter() = delete;
+	~FNetBitArrayPrinter() = delete;
+};
+
+
+//*************************************************************************************************
+// FNetBitArrayPrinter 
+// Implementation
+//*************************************************************************************************
+
+template <class T>
+FString FNetBitArrayPrinter::PrintSetSummary(const T& BitArray)
+{
+	return FString::Printf(TEXT("%u bits set"), BitArray.CountSetBits());
+}
+
+template <class T>
+FString FNetBitArrayPrinter::PrintZeroSummary(const T& BitArray)
+{
+	return FString::Printf(TEXT("%u bits zero"), BitArray.GetNumBits() - BitArray.CountSetBits());
+}
+
+template <class T>
+FString FNetBitArrayPrinter::PrintSetBits(const T& BitArray)
+{
+	FString Buffer;
+	Buffer.Reserve(1024);
+
+	const uint32 BitEndIndex = BitArray.GetNumBits() - 1U;
+	const uint32 WordEndIndex = BitEndIndex / T::WordBitCount;
+	const T::StorageWordType EndWordMask = (~T::StorageWordType(0) >> (~BitEndIndex & (T::WordBitCount - 1U)));
+
+	const T::StorageWordType* BitWordData = BitArray.GetData();
+
+	int32 TotalSetBits = 0;
+	for (uint32 WordIndex = 0; WordIndex < BitArray.GetNumWords(); ++WordIndex)
+	{
+		const bool bIsEndWord = (WordIndex == WordEndIndex);
+		const T::StorageWordType WordMask = bIsEndWord ? EndWordMask : ~T::StorageWordType(0);
+
+		const T::StorageWordType BitWord = BitWordData[WordIndex];
+
+		const int32 NumSetBits = FPlatformMath::CountBits(BitWord & WordMask);
+
+		if (NumSetBits > 0)
+		{
+			TotalSetBits += NumSetBits;
+
+			Buffer += FString::Printf(TEXT(" (0x%.2x)[%#.8x]"), WordIndex, BitWord);
+		}
+	}
+
+	if (Buffer.IsEmpty())
+	{
+		Buffer = TEXT(" none");
+	}
+
+	return FString::Printf(TEXT("%d bits set:%s"), TotalSetBits, *Buffer);
+}
+
+template <class T>
+FString FNetBitArrayPrinter::PrintZeroBits(const T& BitArray)
+{
+	FString Buffer;
+	Buffer.Reserve(1024);
+
+	const uint32 BitEndIndex = BitArray.GetNumBits() - 1U;
+	const uint32 WordEndIndex = BitEndIndex / T::WordBitCount;
+	const T::StorageWordType EndWordMask = (~T::StorageWordType(0) >> (~BitEndIndex & (T::WordBitCount - 1U)));
+
+	const T::StorageWordType* BitWordData = BitArray.GetData();
+
+	int32 TotalZeroBits = 0;
+	for (uint32 WordIndex = 0; WordIndex < BitArray.GetNumWords(); ++WordIndex)
+	{
+		const bool bIsEndWord = (WordIndex == WordEndIndex);
+		const T::StorageWordType WordMask = bIsEndWord ? EndWordMask : ~T::StorageWordType(0);
+
+		const T::StorageWordType BitWord = BitWordData[WordIndex];
+		const T::StorageWordType FlipBitWord = ~(BitWord & WordMask);
+
+		const int32 NumZeroBits = FPlatformMath::CountBits(FlipBitWord);
+
+		if (NumZeroBits > 0)
+		{
+			TotalZeroBits += NumZeroBits;
+
+			Buffer += FString::Printf(TEXT(" (0x%.2x)[%#.8x]"), WordIndex, BitWord);
+		}
+	}
+
+	if (Buffer.IsEmpty())
+	{
+		Buffer = TEXT(" none");
+	}
+
+	return FString::Printf(TEXT("%d bits zero:%s"), TotalZeroBits, *Buffer);
+}
+
+template <class T, class U>
+FString FNetBitArrayPrinter::PrintDeltaSummary(const T& BitArrayA, const U& BitArrayB)
+{
+	UE_NETBITARRAY_VALIDATE_BOTH_COMPATIBLE(BitArrayA, BitArrayB);
+
+	const T::StorageWordType* BitWordDataA = BitArrayA.GetData();
+	const U::StorageWordType* BitWordDataB = BitArrayB.GetData();
+
+	int32 DiffNumBits = 0;
+	for (uint32 WordIndex = 0; WordIndex < BitArrayA.GetNumWords(); ++WordIndex)
+	{
+		const T::StorageWordType BitWordA = BitWordDataA[WordIndex];
+		const U::StorageWordType BitWordB = BitWordDataB[WordIndex];
+
+		// XOR both words
+		const T::StorageWordType DeltaBits = BitWordA ^ BitWordB;
+
+		DiffNumBits += FPlatformMath::CountBits(DeltaBits);
+	}
+
+	return FString::Printf(TEXT("%d bits differ"), DiffNumBits);
+}
+
+template <class T, class U>
+FString FNetBitArrayPrinter::PrintDeltaBits(const T& BitArrayA, const U& BitArrayB)
+{
+	UE_NETBITARRAY_VALIDATE_BOTH_COMPATIBLE(BitArrayA, BitArrayB);
+
+	FString Buffer;
+	Buffer.Reserve(1024);
+
+	const T::StorageWordType* BitWordDataA = BitArrayA.GetData();
+	const U::StorageWordType* BitWordDataB = BitArrayB.GetData();
+
+	int32 DiffNumBits = 0;
+	for (uint32 WordIndex = 0; WordIndex < BitArrayA.GetNumWords(); ++WordIndex)
+	{
+		const T::StorageWordType BitWordA = BitWordDataA[WordIndex];
+		const U::StorageWordType BitWordB = BitWordDataB[WordIndex];
+
+		const T::StorageWordType DeltaBits = BitWordA ^ BitWordB;
+
+		const int32 DiffBits = FPlatformMath::CountBits(DeltaBits);
+
+		if (DiffBits > 0)
+		{
+			DiffNumBits += DiffBits;
+
+			Buffer += FString::Printf(TEXT(" (0x%.2x)[%#.8x]vs[%#.8x]"), WordIndex, BitWordA, BitWordB);
+		}
+	}
+
+	if (Buffer.IsEmpty())
+	{
+		Buffer = TEXT(" none");
+	}
+
+	return FString::Printf(TEXT("%d bits differ:%s"), DiffNumBits, *Buffer);
+}
+
+} // end namespace UE::Net
\ No newline at end of file
