Fix template compile issues on clang
#rnx

diff --git a/Engine/Source/Runtime/Net/Core/Public/Net/Core/NetBitArrayPrinter.h b/Engine/Source/Runtime/Net/Core/Public/Net/Core/NetBitArrayPrinter.h
--- a/Engine/Source/Runtime/Net/Core/Public/Net/Core/NetBitArrayPrinter.h
+++ b/Engine/Source/Runtime/Net/Core/Public/Net/Core/NetBitArrayPrinter.h
@@ -3,6 +3,7 @@
 #pragma once
 
 #include "Net/Core/NetBitArray.h"
+#include "Containers/UnrealString.h"
 
 namespace UE::Net
 {
@@ -66,22 +67,24 @@ FString FNetBitArrayPrinter::PrintZeroSummary(const T& BitArray)
 template <class T>
 FString FNetBitArrayPrinter::PrintSetBits(const T& BitArray)
 {
+	using WordType = typename T::StorageWordType;
+
 	FString Buffer;
 	Buffer.Reserve(1024);
 
 	const uint32 BitEndIndex = BitArray.GetNumBits() - 1U;
 	const uint32 WordEndIndex = BitEndIndex / T::WordBitCount;
-	const T::StorageWordType EndWordMask = (~T::StorageWordType(0) >> (~BitEndIndex & (T::WordBitCount - 1U)));
+	const WordType EndWordMask = (~WordType(0) >> (~BitEndIndex & (T::WordBitCount - 1U)));
 
-	const T::StorageWordType* BitWordData = BitArray.GetData();
+	const WordType* BitWordData = BitArray.GetData();
 
 	int32 TotalSetBits = 0;
 	for (uint32 WordIndex = 0; WordIndex < BitArray.GetNumWords(); ++WordIndex)
 	{
 		const bool bIsEndWord = (WordIndex == WordEndIndex);
-		const T::StorageWordType WordMask = bIsEndWord ? EndWordMask : ~T::StorageWordType(0);
+		const WordType WordMask = bIsEndWord ? EndWordMask : ~WordType(0);
 
-		const T::StorageWordType BitWord = BitWordData[WordIndex];
+		const WordType BitWord = BitWordData[WordIndex];
 
 		const int32 NumSetBits = FPlatformMath::CountBits(BitWord & WordMask);
 
@@ -104,23 +107,25 @@ FString FNetBitArrayPrinter::PrintSetBits(const T& BitArray)
 template <class T>
 FString FNetBitArrayPrinter::PrintZeroBits(const T& BitArray)
 {
+	using WordType = typename T::StorageWordType;
+
 	FString Buffer;
 	Buffer.Reserve(1024);
 
 	const uint32 BitEndIndex = BitArray.GetNumBits() - 1U;
 	const uint32 WordEndIndex = BitEndIndex / T::WordBitCount;
-	const T::StorageWordType EndWordMask = (~T::StorageWordType(0) >> (~BitEndIndex & (T::WordBitCount - 1U)));
+	const WordType EndWordMask = (~WordType(0) >> (~BitEndIndex & (T::WordBitCount - 1U)));
 
-	const T::StorageWordType* BitWordData = BitArray.GetData();
+	const WordType* BitWordData = BitArray.GetData();
 
 	int32 TotalZeroBits = 0;
 	for (uint32 WordIndex = 0; WordIndex < BitArray.GetNumWords(); ++WordIndex)
 	{
 		const bool bIsEndWord = (WordIndex == WordEndIndex);
-		const T::StorageWordType WordMask = bIsEndWord ? EndWordMask : ~T::StorageWordType(0);
+		const WordType WordMask = bIsEndWord ? EndWordMask : ~WordType(0);
 
-		const T::StorageWordType BitWord = BitWordData[WordIndex];
-		const T::StorageWordType FlipBitWord = ~(BitWord & WordMask);
+		const WordType BitWord = BitWordData[WordIndex];
+		const WordType FlipBitWord = ~(BitWord & WordMask);
 
 		const int32 NumZeroBits = FPlatformMath::CountBits(FlipBitWord);
 
@@ -143,19 +148,21 @@ FString FNetBitArrayPrinter::PrintZeroBits(const T& BitArray)
 template <class T, class U>
 FString FNetBitArrayPrinter::PrintDeltaSummary(const T& BitArrayA, const U& BitArrayB)
 {
+	using WordTypeT = typename T::StorageWordType;
+	using WordTypeU = typename U::StorageWordType;
 	UE_NETBITARRAY_VALIDATE_BOTH_COMPATIBLE(BitArrayA, BitArrayB);
 
-	const T::StorageWordType* BitWordDataA = BitArrayA.GetData();
-	const U::StorageWordType* BitWordDataB = BitArrayB.GetData();
+	const WordTypeT* BitWordDataA = BitArrayA.GetData();
+	const WordTypeU* BitWordDataB = BitArrayB.GetData();
 
 	int32 DiffNumBits = 0;
 	for (uint32 WordIndex = 0; WordIndex < BitArrayA.GetNumWords(); ++WordIndex)
 	{
-		const T::StorageWordType BitWordA = BitWordDataA[WordIndex];
-		const U::StorageWordType BitWordB = BitWordDataB[WordIndex];
+		const WordTypeT BitWordA = BitWordDataA[WordIndex];
+		const WordTypeU BitWordB = BitWordDataB[WordIndex];
 
 		// XOR both words
-		const T::StorageWordType DeltaBits = BitWordA ^ BitWordB;
+		const WordTypeT DeltaBits = BitWordA ^ BitWordB;
 
 		DiffNumBits += FPlatformMath::CountBits(DeltaBits);
 	}
@@ -166,21 +173,23 @@ FString FNetBitArrayPrinter::PrintDeltaSummary(const T& BitArrayA, const U& BitA
 template <class T, class U>
 FString FNetBitArrayPrinter::PrintDeltaBits(const T& BitArrayA, const U& BitArrayB)
 {
+	using WordTypeT = typename T::StorageWordType;
+	using WordTypeU = typename U::StorageWordType;
 	UE_NETBITARRAY_VALIDATE_BOTH_COMPATIBLE(BitArrayA, BitArrayB);
 
 	FString Buffer;
 	Buffer.Reserve(1024);
 
-	const T::StorageWordType* BitWordDataA = BitArrayA.GetData();
-	const U::StorageWordType* BitWordDataB = BitArrayB.GetData();
+	const WordTypeT* BitWordDataA = BitArrayA.GetData();
+	const WordTypeU* BitWordDataB = BitArrayB.GetData();
 
 	int32 DiffNumBits = 0;
 	for (uint32 WordIndex = 0; WordIndex < BitArrayA.GetNumWords(); ++WordIndex)
 	{
-		const T::StorageWordType BitWordA = BitWordDataA[WordIndex];
-		const U::StorageWordType BitWordB = BitWordDataB[WordIndex];
+		const WordTypeT BitWordA = BitWordDataA[WordIndex];
+		const WordTypeU BitWordB = BitWordDataB[WordIndex];
 
-		const T::StorageWordType DeltaBits = BitWordA ^ BitWordB;
+		const WordTypeT DeltaBits = BitWordA ^ BitWordB;
 
 		const int32 DiffBits = FPlatformMath::CountBits(DeltaBits);
 
