Iris
* Fix CombineMultiple when using two different WordOp functions
* Properly initialize FInternalNetRefIndex

#rnx

diff --git a/Engine/Source/Runtime/Experimental/Iris/Core/Private/Iris/Core/IrisDebugging.cpp b/Engine/Source/Runtime/Experimental/Iris/Core/Private/Iris/Core/IrisDebugging.cpp
--- a/Engine/Source/Runtime/Experimental/Iris/Core/Private/Iris/Core/IrisDebugging.cpp
+++ b/Engine/Source/Runtime/Experimental/Iris/Core/Private/Iris/Core/IrisDebugging.cpp
@@ -57,7 +57,7 @@ namespace IrisDebugHelperInternal
 			}
 		}));
 
-	static UE::Net::Private::FInternalNetRefIndex GIrisDebugInternalIndex;
+	static UE::Net::Private::FInternalNetRefIndex GIrisDebugInternalIndex = UE::Net::Private::FNetRefHandleManager::InvalidInternalIndex;
 	static FAutoConsoleCommand NetIrisDebugNetInternalIndex(
 		TEXT("Net.Iris.DebugNetInternalIndex"),
 		TEXT("Specify an internal index that we will break on (or none to turn off)."),
diff --git a/Engine/Source/Runtime/Net/Core/Public/Net/Core/NetBitArray.h b/Engine/Source/Runtime/Net/Core/Public/Net/Core/NetBitArray.h
--- a/Engine/Source/Runtime/Net/Core/Public/Net/Core/NetBitArray.h
+++ b/Engine/Source/Runtime/Net/Core/Public/Net/Core/NetBitArray.h
@@ -209,7 +209,7 @@ public:
 	 * @param ArrayB The second bit array to execute the other operation on
 	 * @see FNetBitArrayHelper::AndOp, FNetBitArrayHelper::AndNotOp, FNetBitArrayHelper::OrOp, FNetBitArrayHelper::XorOp
 	 */
-	template<typename WordOpFunctor> void CombineMultiple(WordOpFunctor&& Op, const FNetBitArray& ArrayA, WordOpFunctor&& Op2, const FNetBitArray& ArrayB);
+	template<typename WordOpFunctor1, typename WordOpFunctor2> void CombineMultiple(WordOpFunctor1&& Op, const FNetBitArray& ArrayA, WordOpFunctor2&& Op2, const FNetBitArray& ArrayB);
 
 	/** Iterate over all set bits and invoke functor with signature void(uint32 BitIndex). */
 	template<typename T>
@@ -373,7 +373,7 @@ public:
 	 * @param ArrayB The second bit array to execute the other operation on
 	 * @see FNetBitArrayHelper::AndOp, FNetBitArrayHelper::AndNotOp, FNetBitArrayHelper::OrOp, FNetBitArrayHelper::XorOp
 	 */
-	template<typename WordOpFunctor> void CombineMultiple(WordOpFunctor&& Op, const FNetBitArrayView& ArrayA, WordOpFunctor&& Op2, const FNetBitArrayView& ArrayB);
+	template<typename WordOpFunctor1, typename WordOpFunctor2> void CombineMultiple(WordOpFunctor1&& Op, const FNetBitArrayView& ArrayA, WordOpFunctor2&& Op2, const FNetBitArrayView& ArrayB);
 
 	/** Iterate over all set bits and invoke functor with signature void(uint32 BitIndex). */
 	template <typename T>
@@ -675,7 +675,7 @@ private:
 		}
 	}
 
-	template<typename Functor> static void CombineMultiple(StorageWordType* WriteStorage, const StorageWordType* const ReadStorageA, const StorageWordType* const ReadStorageB, const uint32 WordCount, Functor&& WordOpWrite, Functor&& WordOpRead)
+	template<typename WordOpFunctor1, typename WordOpFunctor2> static void CombineMultiple(StorageWordType* WriteStorage, const StorageWordType* const ReadStorageA, const StorageWordType* const ReadStorageB, const uint32 WordCount, WordOpFunctor1&& WordOpWrite, WordOpFunctor2&& WordOpRead)
 	{
 		for (uint32 WordIt = 0; WordIt < WordCount; ++WordIt)
 		{
@@ -993,7 +993,7 @@ template<typename WordOpFunctor> inline void FNetBitArray::Combine(const FNetBit
 	FNetBitArrayHelper::Combine(Storage.GetData(), Other.Storage.GetData(), Storage.Num(), WordOp);
 }
 
-template<typename WordOpFunctor> inline void FNetBitArray::CombineMultiple(WordOpFunctor&& Op, const FNetBitArray& ArrayA, WordOpFunctor&& Op2, const FNetBitArray& ArrayB)
+template<typename WordOpFunctor1, typename WordOpFunctor2> inline void FNetBitArray::CombineMultiple(WordOpFunctor1&& Op, const FNetBitArray& ArrayA, WordOpFunctor2&& Op2, const FNetBitArray& ArrayB)
 {
 	UE_NETBITARRAY_VALIDATE_COMPATIBLE(ArrayA);
 	UE_NETBITARRAY_VALIDATE_COMPATIBLE(ArrayB);
@@ -1216,7 +1216,7 @@ template<typename WordOpFunctor> inline void FNetBitArrayView::Combine(const FNe
 	FNetBitArrayHelper::Combine(Storage, Other.Storage, WordCount, WordOp);
 }
 
-template<typename WordOpFunctor> inline void FNetBitArrayView::CombineMultiple(WordOpFunctor&& Op, const FNetBitArrayView& ArrayA, WordOpFunctor&& Op2, const FNetBitArrayView& ArrayB)
+template<typename WordOpFunctor1, typename WordOpFunctor2> inline void FNetBitArrayView::CombineMultiple(WordOpFunctor1&& Op, const FNetBitArrayView& ArrayA, WordOpFunctor2&& Op2, const FNetBitArrayView& ArrayB)
 {
 	UE_NETBITARRAYVIEW_VALIDATE_COMPATIBLE(ArrayA);
 	UE_NETBITARRAYVIEW_VALIDATE_COMPATIBLE(ArrayB);
