fixed missing export markup

#rnx

diff --git a/Engine/Source/Runtime/Net/Core/Private/Net/Core/Misc/NetContext.cpp b/Engine/Source/Runtime/Net/Core/Private/Net/Core/Misc/NetContext.cpp
--- a/Engine/Source/Runtime/Net/Core/Private/Net/Core/Misc/NetContext.cpp
+++ b/Engine/Source/Runtime/Net/Core/Private/Net/Core/Misc/NetContext.cpp
@@ -1,6 +1,7 @@
 // Copyright Epic Games, Inc. All Rights Reserved.
 
 #include "Net/Core/Misc/NetContext.h"
+#include "Misc/AssertionMacros.h"
 
 namespace UE::Net
 {
diff --git a/Engine/Source/Runtime/Net/Core/Public/Net/Core/Misc/NetContext.h b/Engine/Source/Runtime/Net/Core/Public/Net/Core/Misc/NetContext.h
--- a/Engine/Source/Runtime/Net/Core/Public/Net/Core/Misc/NetContext.h
+++ b/Engine/Source/Runtime/Net/Core/Public/Net/Core/Misc/NetContext.h
@@ -2,6 +2,8 @@
 
 #pragma once
 
+#include "HAL/PlatformMisc.h"
+
 // Forward declarations
 namespace UE::Net::Private
 {
@@ -13,7 +15,7 @@ namespace UE::Net
 {
 
 /** Stateless class that provides misc network context information */
-class FNetContext
+class NETCORE_API FNetContext
 {
 public:
 	
