NetBitArray
* Added InitAndCopy functions on FNetBitArray. Used to quickly construct a bit array from an existing source (either bitarray or bitarrayview)
* Added Set functions.  Used to overwrite the bitarray with the word operations of two other arrays.

diff --git a/Engine/Plugins/Runtime/ReplicationSystemTestPlugin/Source/Private/Tests/Core/TestNetBitArray.cpp b/Engine/Plugins/Runtime/ReplicationSystemTestPlugin/Source/Private/Tests/Core/TestNetBitArray.cpp
--- a/Engine/Plugins/Runtime/ReplicationSystemTestPlugin/Source/Private/Tests/Core/TestNetBitArray.cpp
+++ b/Engine/Plugins/Runtime/ReplicationSystemTestPlugin/Source/Private/Tests/Core/TestNetBitArray.cpp
@@ -1298,6 +1298,48 @@ UE_NET_TEST(FNetBitArrayView, TestCombineMultiple)
 	}
 }
 
+UE_NET_TEST(FNetBitArrayView, TestSet)
+{
+	// Test Set function on BitArrayView
+	{
+		const uint32 ExpectedWordBuffer[] = { 0x0011, 0x0110, 0x1100, 0x1001 };
+		uint32 WordBufferA[] = { 0xFFFF, 0xFFFF, 0xFFFF, 0xFFFF };
+		uint32 WordBufferB[] = { 0x0001, 0x0010, 0x0100, 0x1000 };
+		uint32 WordBufferC[] = { 0x0010, 0x0100, 0x1000, 0x0001 };
+
+		FNetBitArrayView BitArrayA(&WordBufferA[0], 128);
+		const FNetBitArrayView BitArrayB(&WordBufferB[0], 128);
+		const FNetBitArrayView BitArrayC(&WordBufferC[0], 128);
+
+		BitArrayA.Set(BitArrayB, FNetBitArrayView::OrOp, BitArrayC);
+
+		for (uint32 it = 0; it < UE_ARRAY_COUNT(ExpectedWordBuffer); ++it)
+		{
+			UE_NET_ASSERT_EQ(ExpectedWordBuffer[it], BitArrayA.GetData()[it]);
+		}
+	}
+
+	// Test Set function on BitArray
+	{
+		const uint32 ExpectedWordBuffer[] = { 0x0011, 0x0110, 0x1100, 0x1001 };
+		uint32 WordBufferA[] = { 0xFFFF, 0xFFFF, 0xFFFF, 0xFFFF };
+		uint32 WordBufferB[] = { 0x0001, 0x0010, 0x0100, 0x1000 };
+		uint32 WordBufferC[] = { 0x0010, 0x0100, 0x1000, 0x0001 };
+
+		FNetBitArray BitArrayA, BitArrayB, BitArrayC;
+		BitArrayA.InitAndCopy(FNetBitArrayView(&WordBufferA[0], 128));
+		BitArrayB.InitAndCopy(FNetBitArrayView(&WordBufferB[0], 128));
+		BitArrayC.InitAndCopy(FNetBitArrayView(&WordBufferC[0], 128));
+
+		BitArrayA.Set(BitArrayB, FNetBitArrayView::OrOp, BitArrayC);
+
+		for (uint32 it = 0; it < UE_ARRAY_COUNT(ExpectedWordBuffer); ++it)
+		{
+			UE_NET_ASSERT_EQ(ExpectedWordBuffer[it], BitArrayA.GetData()[it]);
+		}
+	}
+}
+
 UE_NET_TEST(FNetBitArrayView, Copy)
 {	
 	// Test Partial word
@@ -1331,6 +1373,82 @@ UE_NET_TEST(FNetBitArrayView, Copy)
 			UE_NET_ASSERT_EQ(WordBufferB[it], WordBufferA[it]);
 		}
 	}
+
+	// Test View copying BitArray
+	{
+		uint32 WordBufferA[] = { 1, 2, 3, 4 };
+		FNetBitArrayView BitArrayA(&WordBufferA[0], 128);
+		
+		const uint32 WordBufferB[] = { 5, 6, 7, 8 };
+		FNetBitArray BitArrayB;
+		BitArrayB.Init(128);
+		BitArrayB.GetData()[0] = WordBufferB[0];
+		BitArrayB.GetData()[1] = WordBufferB[1];
+		BitArrayB.GetData()[2] = WordBufferB[2];
+		BitArrayB.GetData()[3] = WordBufferB[3];
+
+		BitArrayA.Copy(BitArrayB);
+
+		for (uint32 it = 0; it < UE_ARRAY_COUNT(WordBufferA); ++it)
+		{
+			UE_NET_ASSERT_EQ(BitArrayB.GetData()[it], WordBufferA[it]);
+		}
+	}
+
+	// Test BitArray copying View
+	{
+		uint32 WordBufferA[] = { 1, 2, 3, 4 };
+		FNetBitArrayView BitArrayA(&WordBufferA[0], 128);
+
+		const uint32 WordBufferB[] = { 5, 6, 7, 8 };
+		FNetBitArray BitArrayB;
+		BitArrayB.Init(128);
+		BitArrayB.GetData()[0] = WordBufferB[0];
+		BitArrayB.GetData()[1] = WordBufferB[1];
+		BitArrayB.GetData()[2] = WordBufferB[2];
+		BitArrayB.GetData()[3] = WordBufferB[3];
+
+		BitArrayB.Copy(BitArrayA);
+
+		for (uint32 it = 0; it < UE_ARRAY_COUNT(WordBufferA); ++it)
+		{
+			UE_NET_ASSERT_EQ(BitArrayB.GetData()[it], WordBufferA[it]);
+		}
+	}
+}
+
+UE_NET_TEST(FNetBitArrayView, InitAndCopy)
+{
+	{
+		uint32 WordBufferA[] = { 1, 2, 3, 4 };
+		const FNetBitArrayView BitArrayA(&WordBufferA[0], 128);
+
+		FNetBitArray TestBitArrayA;
+		TestBitArrayA.InitAndCopy(BitArrayA);
+
+		for (uint32 it = 0; it < UE_ARRAY_COUNT(WordBufferA); ++it)
+		{
+			UE_NET_ASSERT_EQ(TestBitArrayA.GetData()[it], WordBufferA[it]);
+		}
+	}
+
+	{
+		uint32 WordBufferB[] = { 5, 6, 7, 8 };
+		FNetBitArray BitArrayB;
+		BitArrayB.Init(128);
+		BitArrayB.GetData()[0] = WordBufferB[0];
+		BitArrayB.GetData()[1] = WordBufferB[1];
+		BitArrayB.GetData()[2] = WordBufferB[2];
+		BitArrayB.GetData()[3] = WordBufferB[3];
+
+		FNetBitArray TestBitArrayB;
+		TestBitArrayB.InitAndCopy(BitArrayB);
+
+		for (uint32 it = 0; it < UE_ARRAY_COUNT(WordBufferB); ++it)
+		{
+			UE_NET_ASSERT_EQ(TestBitArrayB.GetData()[it], WordBufferB[it]);
+		}
+	}
 }
 
 UE_NET_TEST(FNetBitArrayView, PrintHelper)
diff --git a/Engine/Source/Runtime/Net/Core/Public/Net/Core/NetBitArray.h b/Engine/Source/Runtime/Net/Core/Public/Net/Core/NetBitArray.h
--- a/Engine/Source/Runtime/Net/Core/Public/Net/Core/NetBitArray.h
+++ b/Engine/Source/Runtime/Net/Core/Public/Net/Core/NetBitArray.h
@@ -36,9 +36,7 @@ namespace UE::Net
 #	define UE_NETBITARRAY_CHECK(...)
 #endif
 
-#define UE_NETBITARRAY_VALIDATE_COMPATIBLE(x) UE_NETBITARRAY_CHECK(this != &(x) && BitCount == (x).BitCount)
-#define UE_NETBITARRAY_VALIDATE_BOTH_COMPATIBLE(lhs,rhs) UE_NETBITARRAY_CHECK((void*)(&lhs) != (void*)(&rhs) && (lhs).GetNumBits() == (rhs).GetNumBits())
-#define UE_NETBITARRAYVIEW_VALIDATE_COMPATIBLE(x) UE_NETBITARRAY_CHECK((Storage != (x).Storage) && (BitCount == (x).BitCount))
+#define UE_NETBITARRAY_VALIDATE_BOTH_COMPATIBLE(lhs,rhs) UE_NETBITARRAY_CHECK((lhs).GetData() != (rhs).GetData() && (lhs).GetNumBits() == (rhs).GetNumBits())
 
 namespace UE::Net
 {
@@ -81,6 +79,11 @@ public:
 	 */
 	void Init(uint32 BitCount);
 
+	/**
+	 * Initialize a bit array with another array.
+	 */
+	void InitAndCopy(const FNetBitArray& Source);
+	void InitAndCopy(const FNetBitArrayView& Source);
 	/**
 	 * Release all memory and set capacity to 0 bits
 	 */
@@ -189,6 +192,17 @@ public:
 
 	/** Copy bits from other bit array. The other array must have the same bit count. */
 	void Copy(const FNetBitArray& Other);
+	void Copy(const FNetBitArrayView& Other);
+
+	/**
+	 * Overwrite this bit array with the result of a word operation from two other arrays
+	 * 
+	 * @param First The first bit array to pass in the word operation
+	 * @param WordOp The operation to perform on the two bit arrays
+	 * @param Second The second bit array to pass in the word operation
+	 * @see FNetBitArrayHelper::AndOp, FNetBitArrayHelper::AndNotOp, FNetBitArrayHelper::OrOp, FNetBitArrayHelper::XorOp
+	 */
+	template<typename WordOpFunctor> void Set(const FNetBitArray& First, WordOpFunctor&& WordOp, const FNetBitArray& Second);
 
 	/**
 	 * Combine this array with another array using a word operation functor. Bit counts need to match.
@@ -353,6 +367,17 @@ public:
 
 	/** Copy bits from other bit array. The other array must have the same bit count. */
 	inline void Copy(const FNetBitArrayView& Other);
+	inline void Copy(const FNetBitArray& Other);
+
+	/**
+	 * Overwrite this bit array with the result of a word operation from two other arrays
+	 *
+	 * @param First The first bit array to pass in the word operation
+	 * @param WordOp The operation to perform on the two bit arrays
+	 * @param Second The second bit array to pass in the word operation
+	 * @see FNetBitArrayHelper::AndOp, FNetBitArrayHelper::AndNotOp, FNetBitArrayHelper::OrOp, FNetBitArrayHelper::XorOp
+	 */
+	template<typename WordOpFunctor> void Set(const FNetBitArrayView& First, WordOpFunctor&& WordOp, const FNetBitArrayView& Second);
 
 	/**
 	 * Combine this array with another array using a word operation functor. Bit counts need to match.
@@ -667,6 +692,14 @@ private:
 	
 	NETCORE_API static uint32 CountSetBits(const StorageWordType* Storage, const uint32 BitCount, const uint32 StartIndex, const uint32 Count);
 
+	template<typename WordOpFunctor> static void Set(StorageWordType* Storage, const StorageWordType* FirstSource, const StorageWordType* SecondSource, const uint32 WordCount, WordOpFunctor&& WordOp)
+	{
+		for (uint32 WordIt = 0; WordIt < WordCount; ++WordIt)
+		{
+			Storage[WordIt] = WordOp(FirstSource[WordIt], SecondSource[WordIt]);
+		}
+	}
+
 	template<typename Functor> static void Combine(StorageWordType* Storage, const StorageWordType* OtherStorage, const uint32 WordCount, Functor&& WordOp)
 	{
 		for (uint32 WordIt = 0; WordIt < WordCount; ++WordIt)
@@ -853,11 +886,31 @@ inline void FNetBitArray::Init(uint32 InBitCount)
 	Storage.AddZeroed(WordCount);
 }
 
+inline void FNetBitArray::InitAndCopy(const FNetBitArray& Source)
+{
+	BitCount = Source.GetNumBits();
+	const uint32 WordCount = (BitCount + WordBitCount - 1U) / WordBitCount;
+	Storage.Reset(WordCount);
+	Storage.AddUninitialized(WordCount);
+
+	Copy(Source);
+}
+
+inline void FNetBitArray::InitAndCopy(const FNetBitArrayView& Source)
+{
+	BitCount = Source.GetNumBits();
+	const uint32 WordCount = (BitCount + WordBitCount - 1U) / WordBitCount;
+	Storage.Reset(WordCount);
+	Storage.AddUninitialized(WordCount);
+
+	Copy(Source);
+}
+
 inline void FNetBitArray::Empty()
 {
 	BitCount = 0U;
 	Storage.Empty();
-};
+}
 
 inline void FNetBitArray::SetNumBits(uint32 InBitCount)
 {
@@ -982,21 +1035,41 @@ inline uint32 FNetBitArray::CountSetBits(uint32 StartIndex, uint32 Count) const
 
 inline void FNetBitArray::Copy(const FNetBitArray& Other)
 {
-	UE_NETBITARRAY_VALIDATE_COMPATIBLE(Other);
-	// Intentionally doing a memory copy instead of using assignment operator as we know the bit counts are the same and don't need extra checks.
-	FPlatformMemory::Memcpy(GetData(), Other.GetData(), GetNumWords()*sizeof(StorageWordType));
+	UE_NETBITARRAY_VALIDATE_BOTH_COMPATIBLE(*this, Other);
+	if (GetNumWords() > 0 && Other.GetNumWords() > 0)
+	{
+		// Intentionally doing a memory copy instead of using assignment operator as we know the bit counts are the same and don't need extra checks.
+		FPlatformMemory::Memcpy(GetData(), Other.GetData(), GetNumWords() * sizeof(StorageWordType));
+	}
+}
+
+inline void FNetBitArray::Copy(const FNetBitArrayView& Other)
+{
+	UE_NETBITARRAY_VALIDATE_BOTH_COMPATIBLE(*this, Other);
+	if (GetNumWords() > 0 && Other.GetNumWords() > 0)
+	{
+		// Intentionally doing a memory copy instead of using assignment operator as we know the bit counts are the same and don't need extra checks.
+		FPlatformMemory::Memcpy(GetData(), Other.GetData(), GetNumWords() * sizeof(StorageWordType));
+	}
+}
+
+template<typename WordOpFunctor> inline void FNetBitArray::Set(const FNetBitArray& First, WordOpFunctor&& WordOp, const FNetBitArray& Second)
+{
+	UE_NETBITARRAY_VALIDATE_BOTH_COMPATIBLE(*this, First);
+	UE_NETBITARRAY_VALIDATE_BOTH_COMPATIBLE(*this, Second);
+	FNetBitArrayHelper::Set(GetData(), First.GetData(), Second.GetData(), GetNumWords(), WordOp);
 }
 
 template<typename WordOpFunctor> inline void FNetBitArray::Combine(const FNetBitArray& Other, WordOpFunctor&& WordOp)
 {
-	UE_NETBITARRAY_VALIDATE_COMPATIBLE(Other);
+	UE_NETBITARRAY_VALIDATE_BOTH_COMPATIBLE(*this, Other);
 	FNetBitArrayHelper::Combine(Storage.GetData(), Other.Storage.GetData(), Storage.Num(), WordOp);
 }
 
 template<typename WordOpFunctor1, typename WordOpFunctor2> inline void FNetBitArray::CombineMultiple(WordOpFunctor1&& Op, const FNetBitArray& ArrayA, WordOpFunctor2&& Op2, const FNetBitArray& ArrayB)
 {
-	UE_NETBITARRAY_VALIDATE_COMPATIBLE(ArrayA);
-	UE_NETBITARRAY_VALIDATE_COMPATIBLE(ArrayB);
+	UE_NETBITARRAY_VALIDATE_BOTH_COMPATIBLE(*this, ArrayA);
+	UE_NETBITARRAY_VALIDATE_BOTH_COMPATIBLE(*this, ArrayB);
 	FNetBitArrayHelper::CombineMultiple(Storage.GetData(), ArrayA.Storage.GetData(), ArrayB.Storage.GetData(), Storage.Num(), Op, Op2);
 }
 
@@ -1123,11 +1196,22 @@ void FNetBitArrayView::SetAllBits()
 	ClearPaddingBits();
 }
 
-/** Copy bits from other bit array */
 void FNetBitArrayView::Copy(const FNetBitArrayView& Other)
 {
-	UE_NETBITARRAYVIEW_VALIDATE_COMPATIBLE(Other);
-	FPlatformMemory::Memcpy(&Storage[0], &Other.Storage[0], WordCount * sizeof(StorageWordType));
+	UE_NETBITARRAY_VALIDATE_BOTH_COMPATIBLE(*this, Other);
+	if (GetNumWords() > 0 && Other.GetNumWords() > 0 )
+	{
+		FPlatformMemory::Memcpy(&Storage[0], &Other.Storage[0], WordCount * sizeof(StorageWordType));
+	}
+}
+
+void FNetBitArrayView::Copy(const FNetBitArray& Other)
+{
+	UE_NETBITARRAY_VALIDATE_BOTH_COMPATIBLE(*this, Other);
+	if (GetNumWords() > 0 && Other.GetNumWords() > 0)
+	{
+		FPlatformMemory::Memcpy(&Storage[0], Other.GetData(), WordCount * sizeof(StorageWordType));
+	}
 }
 
 void FNetBitArrayView::SetBit(uint32 Index)
@@ -1210,16 +1294,23 @@ inline uint32 FNetBitArrayView::CountSetBits(uint32 StartIndex, uint32 Count) co
 	return FNetBitArrayHelper::CountSetBits(Storage, BitCount, StartIndex, Count);
 }
 
+template<typename WordOpFunctor> inline void FNetBitArrayView::Set(const FNetBitArrayView& First, WordOpFunctor&& WordOp, const FNetBitArrayView& Second)
+{
+	UE_NETBITARRAY_VALIDATE_BOTH_COMPATIBLE(*this, First);
+	UE_NETBITARRAY_VALIDATE_BOTH_COMPATIBLE(*this, Second);
+	FNetBitArrayHelper::Set(GetData(), First.GetData(), Second.GetData(), GetNumWords(), WordOp);
+}
+
 template<typename WordOpFunctor> inline void FNetBitArrayView::Combine(const FNetBitArrayView& Other, WordOpFunctor&& WordOp)
 {
-	UE_NETBITARRAYVIEW_VALIDATE_COMPATIBLE(Other);
+	UE_NETBITARRAY_VALIDATE_BOTH_COMPATIBLE(*this, Other);
 	FNetBitArrayHelper::Combine(Storage, Other.Storage, WordCount, WordOp);
 }
 
 template<typename WordOpFunctor1, typename WordOpFunctor2> inline void FNetBitArrayView::CombineMultiple(WordOpFunctor1&& Op, const FNetBitArrayView& ArrayA, WordOpFunctor2&& Op2, const FNetBitArrayView& ArrayB)
 {
-	UE_NETBITARRAYVIEW_VALIDATE_COMPATIBLE(ArrayA);
-	UE_NETBITARRAYVIEW_VALIDATE_COMPATIBLE(ArrayB);
+	UE_NETBITARRAY_VALIDATE_BOTH_COMPATIBLE(*this, ArrayA);
+	UE_NETBITARRAY_VALIDATE_BOTH_COMPATIBLE(*this, ArrayB);
 	FNetBitArrayHelper::CombineMultiple(Storage, ArrayA.Storage, ArrayB.Storage, WordCount, Op, Op2);
 }
 
@@ -1252,11 +1343,13 @@ inline Private::FNetBitArrayRangedForConstIterator FNetBitArrayView::end() const
 	return Private::FNetBitArrayRangedForConstIterator(Storage, BitCount, Private::FNetBitArrayRangedForConstIterator::ERangeStart::End);
 }
 
+/** Transform any raw buffer into a FNetBitArrayView */
 inline FNetBitArrayView MakeNetBitArrayView(const FNetBitArrayView::StorageWordType* Storage, uint32 BitCount)
 {
 	return FNetBitArrayView(const_cast<FNetBitArrayView::StorageWordType*>(Storage), BitCount);
 }
 
+/** Transform a FNetBitArray into a FNetBitArrayView */
 inline FNetBitArrayView MakeNetBitArrayView(const FNetBitArray& BitArray)
 {
 	return FNetBitArrayView(const_cast<FNetBitArrayView::StorageWordType*>(BitArray.GetData()), BitArray.GetNumBits());
