Iris
* Added NoResetNoValidate constructor to TBitArray :  used to construct a bit array with uninitialized memory.
* FReplicationPrioritization will only update dirty and relevant objects instead of all dirty objects

#rnx

diff --git a/Engine/Source/Runtime/Experimental/Iris/Core/Private/Iris/ReplicationSystem/ReplicationSystem.cpp b/Engine/Source/Runtime/Experimental/Iris/Core/Private/Iris/ReplicationSystem/ReplicationSystem.cpp
--- a/Engine/Source/Runtime/Experimental/Iris/Core/Private/Iris/ReplicationSystem/ReplicationSystem.cpp
+++ b/Engine/Source/Runtime/Experimental/Iris/Core/Private/Iris/ReplicationSystem/ReplicationSystem.cpp
@@ -291,13 +291,22 @@ public:
 	}
 
 	// Runs after filtering
-	void UpdatePrioritization(const FNetBitArrayView& ReplicatingConnections, const FNetBitArrayView& DirtyObjects)
+	void UpdatePrioritization(const FNetBitArrayView& ReplicatingConnections)
 	{
 		IRIS_PROFILER_SCOPE(FReplicationSystem::FImpl::UpdatePrioritization);
 		LLM_SCOPE_BYTAG(Iris);
 
+		const FNetBitArrayView RelevantObjects = ReplicationSystemInternal.GetNetRefHandleManager().GetRelevantObjectsInternalIndices();
+
+		// Make a list of objects that were dirty and are also relevant
+		FNetBitArray DirtyAndRelevantObjects(RelevantObjects.GetNumBits(), FNetBitArray::NoResetNoValidate);
+		FNetBitArrayView DirtyAndRelevantObjectsView = MakeNetBitArrayView(DirtyAndRelevantObjects, FNetBitArray::NoResetNoValidate);
+
+		const FNetBitArrayView AccumulatedDirtyObjects = ReplicationSystemInternal.GetDirtyNetObjectTracker().GetAccumulatedDirtyNetObjects();
+		DirtyAndRelevantObjectsView.Set(RelevantObjects, FNetBitArray::AndOp, AccumulatedDirtyObjects);
+
 		FReplicationPrioritization& Prioritization = ReplicationSystemInternal.GetPrioritization();
-		Prioritization.Prioritize(ReplicatingConnections, DirtyObjects);
+		Prioritization.Prioritize(ReplicatingConnections, DirtyAndRelevantObjectsView);
 	}
 
 	void PropagateDirtyChanges()
@@ -618,10 +627,7 @@ void UReplicationSystem::PreSendUpdate(float DeltaSeconds)
 	if (bAllowObjectReplication)
 	{
 		// Update object priorities
-		FDirtyObjectsAccessor DirtyObjectsAccessor(InternalSys.GetDirtyNetObjectTracker());
-		const FNetBitArrayView DirtyObjects = DirtyObjectsAccessor.GetDirtyNetObjects();
-
-		Impl->UpdatePrioritization(ReplicatingConnections, DirtyObjects);
+		Impl->UpdatePrioritization(ReplicatingConnections);
 
 		// Delta compression preparations before send
 		{
diff --git a/Engine/Source/Runtime/Net/Core/Public/Net/Core/NetBitArray.h b/Engine/Source/Runtime/Net/Core/Public/Net/Core/NetBitArray.h
--- a/Engine/Source/Runtime/Net/Core/Public/Net/Core/NetBitArray.h
+++ b/Engine/Source/Runtime/Net/Core/Public/Net/Core/NetBitArray.h
@@ -44,6 +44,9 @@ namespace UE::Net
 class FNetBitArrayBase
 {
 public:
+	enum ENoResetNoValidateType { NoResetNoValidate };
+	enum EResetOnInitType { ResetOnInit };
+
 	typedef uint32 StorageWordType;
 	static constexpr uint32 WordBitCount = sizeof(StorageWordType)*8U;
 
@@ -70,6 +73,9 @@ public:
 	/**  Creates an array with size BitCount bits. */
 	explicit FNetBitArray(uint32 BitCount);
 
+	/** Construct NetBitArray with the given BitCount but doesn't zero the allocated memory. Useful when the entire bit array will be overwritten before being read.*/
+	explicit FNetBitArray(uint32 BitCountIn, const ENoResetNoValidateType);
+
 	/** Return true if equal including BitCount and padding bits */
 	bool operator==(const FNetBitArray& Other) const;
 
@@ -248,6 +254,9 @@ private:
 	StorageWordType GetLastWordMask() const { return (~StorageWordType(0) >> (uint32(-int32(BitCount)) & (WordBitCount - 1))); }
 	void ClearPaddingBits();
 
+	/** Sets a new size but doesn't initialize the contents with zero'd memory. */
+	void Reserve(uint32 BitCount);
+
 	TArray<StorageWordType> Storage;
 	uint32 BitCount;
 };
@@ -261,8 +270,6 @@ private:
 class FNetBitArrayView : public FNetBitArrayBase
 {
 public:
-	enum ENoResetNoValidateType { NoResetNoValidate };
-	enum EResetOnInitType { ResetOnInit };
 
 	/** Constructor will produce a valid but empty bitarray. */
 	inline FNetBitArrayView();
@@ -877,6 +884,12 @@ inline FNetBitArray::FNetBitArray(uint32 InBitCount)
 	Init(InBitCount);
 }
 
+inline FNetBitArray::FNetBitArray(uint32 InBitCount, const FNetBitArrayBase::ENoResetNoValidateType)
+: BitCount(InBitCount)
+{
+	Reserve(BitCount);
+}
+
 inline void FNetBitArray::Init(uint32 InBitCount)
 {
 	BitCount = InBitCount;
@@ -886,22 +899,25 @@ inline void FNetBitArray::Init(uint32 InBitCount)
 	Storage.AddZeroed(WordCount);
 }
 
-inline void FNetBitArray::InitAndCopy(const FNetBitArray& Source)
+inline void FNetBitArray::Reserve(uint32 InBitCount)
 {
-	BitCount = Source.GetNumBits();
-	const uint32 WordCount = (BitCount + WordBitCount - 1U) / WordBitCount;
+	BitCount = InBitCount;
+
+	const uint32 WordCount = (InBitCount + WordBitCount - 1U) / WordBitCount;
 	Storage.Reset(WordCount);
 	Storage.AddUninitialized(WordCount);
+}
+
+inline void FNetBitArray::InitAndCopy(const FNetBitArray& Source)
+{
+	Reserve(Source.GetNumBits());
 
 	Copy(Source);
 }
 
 inline void FNetBitArray::InitAndCopy(const FNetBitArrayView& Source)
 {
-	BitCount = Source.GetNumBits();
-	const uint32 WordCount = (BitCount + WordBitCount - 1U) / WordBitCount;
-	Storage.Reset(WordCount);
-	Storage.AddUninitialized(WordCount);
+	Reserve(Source.GetNumBits());
 
 	Copy(Source);
 }
@@ -1355,6 +1371,13 @@ inline FNetBitArrayView MakeNetBitArrayView(const FNetBitArray& BitArray)
 	return FNetBitArrayView(const_cast<FNetBitArrayView::StorageWordType*>(BitArray.GetData()), BitArray.GetNumBits());
 }
 
+/** Transform a FNetBitArray into a FNetBitArrayView without validating the contents of the bit array received */
+inline FNetBitArrayView MakeNetBitArrayView(FNetBitArray& BitArray, const FNetBitArrayBase::ENoResetNoValidateType)
+{
+	return FNetBitArrayView(const_cast<FNetBitArrayView::StorageWordType*>(BitArray.GetData()), BitArray.GetNumBits(), FNetBitArrayBase::NoResetNoValidate);
+}
+
+
 } // end namespace UE::Net
 
 namespace UE::Net::Private
