Remove FloatCastChecked in FDDoSDetection::PreFrameReceive

- The first tick after forking has the DeltaTime of the parent, if the parent has been alive for > 4 days the child will immediately assert here.

diff --git a/Engine/Source/Runtime/Net/Core/Private/Net/Core/Misc/DDoSDetection.cpp b/Engine/Source/Runtime/Net/Core/Private/Net/Core/Misc/DDoSDetection.cpp
--- a/Engine/Source/Runtime/Net/Core/Private/Net/Core/Misc/DDoSDetection.cpp
+++ b/Engine/Source/Runtime/Net/Core/Private/Net/Core/Misc/DDoSDetection.cpp
@@ -284,7 +284,7 @@ void FDDoSDetection::PreFrameReceive(float DeltaTime)
 		}
 
 		// At the start of every frame, adjust the DDoS detection based upon DeltaTime - unless there is excessive hitching
-		FrameAdjustment = FloatCastChecked<float>((HitchFrameCount > 0 && HitchFrameCount > HitchFrameTolerance) ? 1.0 : (double)DeltaTime / ExpectedFrameTime, UE::LWC::DefaultFloatPrecision);
+		FrameAdjustment = static_cast<float>((HitchFrameCount > 0 && HitchFrameCount > HitchFrameTolerance) ? 1.0 : (double)DeltaTime / ExpectedFrameTime);
 
 		if (ActiveState > 0 && CooloffTime > 0 && (float)(StartFrameRecvTimestamp - LastMetEscalationConditions) > (float)CooloffTime)
 		{
