NetDriver:
* Fix net.EnableDetailedScopeCounters not turning off conditional scopes because of the multiple definitions of  GUseDetailedScopeCounters.

diff --git a/Engine/Source/Runtime/Engine/Private/DataChannel.cpp b/Engine/Source/Runtime/Engine/Private/DataChannel.cpp
--- a/Engine/Source/Runtime/Engine/Private/DataChannel.cpp
+++ b/Engine/Source/Runtime/Engine/Private/DataChannel.cpp
@@ -107,8 +107,6 @@ static FAutoConsoleVariableRef CVarDormancyHysteresis(
 	TEXT("When > 0, represents the time we'll wait before letting a channel become fully dormant (in seconds). This can prevent churn when objects are going in and out of dormant more frequently than normal.")
 );
 
-extern bool GUseDetailedScopeCounters;
-
 namespace UE::Net
 {
 	extern int32 FilterGuidRemapping;
diff --git a/Engine/Source/Runtime/Engine/Private/DataReplication.cpp b/Engine/Source/Runtime/Engine/Private/DataReplication.cpp
--- a/Engine/Source/Runtime/Engine/Private/DataReplication.cpp
+++ b/Engine/Source/Runtime/Engine/Private/DataReplication.cpp
@@ -84,8 +84,6 @@ FAutoConsoleVariableRef CVarPushModelSkipUndirtiedFastArrays(
 
 extern int32 GNumSkippedObjectEmptyUpdates;
 
-extern bool GUseDetailedScopeCounters;
-
 class FNetSerializeCB : public INetSerializeCB
 {
 private:
diff --git a/Engine/Source/Runtime/Engine/Private/RepLayout.cpp b/Engine/Source/Runtime/Engine/Private/RepLayout.cpp
--- a/Engine/Source/Runtime/Engine/Private/RepLayout.cpp
+++ b/Engine/Source/Runtime/Engine/Private/RepLayout.cpp
@@ -116,8 +116,6 @@ namespace UE::Net::Private
 extern int32 GNumSharedSerializationHit;
 extern int32 GNumSharedSerializationMiss;
 
-extern bool GUseDetailedScopeCounters;
-
 /** 
 * Helper method to allow us to instrument FBitArchive using FNetTraceCollector
 * The rationale behind this variant being declared here is that it makes ugly assumption about FBitArchive when being used to collect stats either is a BitWriter or a BitReader
diff --git a/Engine/Source/Runtime/Net/Core/Classes/Net/Serialization/FastArraySerializer.h b/Engine/Source/Runtime/Net/Core/Classes/Net/Serialization/FastArraySerializer.h
--- a/Engine/Source/Runtime/Net/Core/Classes/Net/Serialization/FastArraySerializer.h
+++ b/Engine/Source/Runtime/Net/Core/Classes/Net/Serialization/FastArraySerializer.h
@@ -32,8 +32,6 @@ DECLARE_CYCLE_STAT_EXTERN(TEXT("NetSerializeFast Array"), STAT_NetSerializeFastA
 DECLARE_CYCLE_STAT_EXTERN(TEXT("NetSerializeFast Array BuildMap"), STAT_NetSerializeFastArray_BuildMap, STATGROUP_ServerCPU, NETCORE_API);
 DECLARE_CYCLE_STAT_EXTERN(TEXT("NetSerializeFast Array Delta Struct"), STAT_NetSerializeFastArray_DeltaStruct, STATGROUP_ServerCPU, NETCORE_API);
 
-extern bool GUseDetailedScopeCounters;
-
 /**
  *	===================== Fast TArray Replication ===================== 
  *
diff --git a/Engine/Source/Runtime/Net/Core/Private/Net/Core/NetCoreModule.cpp b/Engine/Source/Runtime/Net/Core/Private/Net/Core/NetCoreModule.cpp
--- a/Engine/Source/Runtime/Net/Core/Private/Net/Core/NetCoreModule.cpp
+++ b/Engine/Source/Runtime/Net/Core/Private/Net/Core/NetCoreModule.cpp
@@ -11,7 +11,7 @@
 #include "Net/Core/NetHandle/NetHandleManager.h"
 #include "Net/Serialization/FastArraySerializer.h"
 
-extern bool GUseDetailedScopeCounters;
+bool GUseDetailedScopeCounters=true;
 static FAutoConsoleVariableRef CVarNetEnableDetailedScopeCounters(
 	TEXT("net.EnableDetailedScopeCounters"),
 	GUseDetailedScopeCounters,
diff --git a/Engine/Source/Runtime/Net/Core/Public/Net/Core/NetCoreModule.h b/Engine/Source/Runtime/Net/Core/Public/Net/Core/NetCoreModule.h
--- a/Engine/Source/Runtime/Net/Core/Public/Net/Core/NetCoreModule.h
+++ b/Engine/Source/Runtime/Net/Core/Public/Net/Core/NetCoreModule.h
@@ -2,4 +2,5 @@
 
 #pragma once
 
-static bool GUseDetailedScopeCounters = true;
\ No newline at end of file
+
+extern NETCORE_API bool GUseDetailedScopeCounters;
\ No newline at end of file
