Fix missing include in non-unity build
#rnx

diff --git a/Engine/Source/Runtime/Net/Core/Private/Net/Core/NetCoreModule.cpp b/Engine/Source/Runtime/Net/Core/Private/Net/Core/NetCoreModule.cpp
--- a/Engine/Source/Runtime/Net/Core/Private/Net/Core/NetCoreModule.cpp
+++ b/Engine/Source/Runtime/Net/Core/Private/Net/Core/NetCoreModule.cpp
@@ -1,7 +1,6 @@
 // Copyright Epic Games, Inc. All Rights Reserved.
 
 #include "Net/Core/NetCoreModule.h"
-#include "CoreMinimal.h"
 #include "Modules/ModuleInterface.h"
 #include "Modules/ModuleManager.h"
 #include "Net/Core/Misc/NetCoreLog.h"
diff --git a/Engine/Source/Runtime/Net/Core/Public/Net/Core/NetCoreModule.h b/Engine/Source/Runtime/Net/Core/Public/Net/Core/NetCoreModule.h
--- a/Engine/Source/Runtime/Net/Core/Public/Net/Core/NetCoreModule.h
+++ b/Engine/Source/Runtime/Net/Core/Public/Net/Core/NetCoreModule.h
@@ -2,5 +2,6 @@
 
 #pragma once
 
+#include "CoreMinimal.h"
 
 extern NETCORE_API bool GUseDetailedScopeCounters;
\ No newline at end of file
