Added MakeNetBitArray from storage with NoResetNoValidate.

#rnx

diff --git a/Engine/Source/Runtime/Net/Core/Public/Net/Core/NetBitArray.h b/Engine/Source/Runtime/Net/Core/Public/Net/Core/NetBitArray.h
--- a/Engine/Source/Runtime/Net/Core/Public/Net/Core/NetBitArray.h
+++ b/Engine/Source/Runtime/Net/Core/Public/Net/Core/NetBitArray.h
@@ -1366,6 +1366,11 @@ inline FNetBitArrayView MakeNetBitArrayView(const FNetBitArrayView::StorageWordT
 	return FNetBitArrayView(const_cast<FNetBitArrayView::StorageWordType*>(Storage), BitCount);
 }
 
+inline FNetBitArrayView MakeNetBitArrayView(const FNetBitArrayView::StorageWordType* Storage, uint32 BitCount, const FNetBitArrayBase::ENoResetNoValidateType)
+{
+	return FNetBitArrayView(const_cast<FNetBitArrayView::StorageWordType*>(Storage), BitCount, FNetBitArrayBase::NoResetNoValidate);
+}
+
 /** Transform a FNetBitArray into a FNetBitArrayView */
 inline FNetBitArrayView MakeNetBitArrayView(const FNetBitArray& BitArray)
 {
