Iris
* Support objects triggering push-model dirty during PreReplication.
* Added locks to the GlobalDirtyNetTracker to detect illegal modifications after the RepSystem has read the list for the last time until the next reset.
* Only initialize the RepSystem's DirtyNetObjectTracker when object replication is enabled (e.g. server)

#rnx

diff --git a/Engine/Source/Runtime/Experimental/Iris/Core/Private/Iris/ReplicationState/IrisFastArraySerializer.cpp b/Engine/Source/Runtime/Experimental/Iris/Core/Private/Iris/ReplicationState/IrisFastArraySerializer.cpp
--- a/Engine/Source/Runtime/Experimental/Iris/Core/Private/Iris/ReplicationState/IrisFastArraySerializer.cpp
+++ b/Engine/Source/Runtime/Experimental/Iris/Core/Private/Iris/ReplicationState/IrisFastArraySerializer.cpp
@@ -107,7 +107,7 @@ void FIrisFastArraySerializerPrivateAccessor::MarkAllArrayItemsDirty(FIrisFastAr
 	FNetBitArrayView MemberChangeMask = UE::Net::Private::FIrisFastArraySerializerPrivateAccessor::GetChangeMask(Array);
 	if (!MemberChangeMask.GetBit(FIrisFastArraySerializer::IrisFastArrayPropertyBitIndex))
 	{
-		MarkNetObjectStateDirty(Array.ReplicationStateHeader);
+		MarkNetObjectStateHeaderDirty(Array.ReplicationStateHeader);
 	}
 	if (StartingIndex == 0)
 	{
@@ -130,7 +130,7 @@ void FIrisFastArraySerializerPrivateAccessor::MarkArrayDirty(FIrisFastArraySeria
 	if (!MemberChangeMask.GetBit(FIrisFastArraySerializer::IrisFastArrayPropertyBitIndex))
 	{
 		MemberChangeMask.SetBit(FIrisFastArraySerializer::IrisFastArrayPropertyBitIndex);
-		MarkNetObjectStateDirty(Array.ReplicationStateHeader);
+		MarkNetObjectStateHeaderDirty(Array.ReplicationStateHeader);
 	}
 }
 
@@ -147,7 +147,7 @@ void FIrisFastArraySerializerPrivateAccessor::MarkArrayItemDirty(FIrisFastArrayS
 	if (!MemberChangeMask.GetBit(FIrisFastArraySerializer::IrisFastArrayPropertyBitIndex))
 	{
 		MemberChangeMask.SetBit(FIrisFastArraySerializer::IrisFastArrayPropertyBitIndex);
-		MarkNetObjectStateDirty(Array.ReplicationStateHeader);
+		MarkNetObjectStateHeaderDirty(Array.ReplicationStateHeader);
 	}
 }
 
diff --git a/Engine/Source/Runtime/Experimental/Iris/Core/Private/Iris/ReplicationState/PropertyReplicationState.cpp b/Engine/Source/Runtime/Experimental/Iris/Core/Private/Iris/ReplicationState/PropertyReplicationState.cpp
--- a/Engine/Source/Runtime/Experimental/Iris/Core/Private/Iris/ReplicationState/PropertyReplicationState.cpp
+++ b/Engine/Source/Runtime/Experimental/Iris/Core/Private/Iris/ReplicationState/PropertyReplicationState.cpp
@@ -207,7 +207,7 @@ void FPropertyReplicationState::MarkDirty(uint32 Index)
 			Private::FReplicationStateHeaderAccessor::MarkInitStateDirty(Header);
 			if (Header.IsBound())
 			{
-				MarkNetObjectStateDirty(Header);
+				MarkNetObjectStateHeaderDirty(Header);
 			}
 		}
 	}
@@ -248,7 +248,7 @@ void FPropertyReplicationState::MarkArrayDirty(uint32 Index)
 			Private::FReplicationStateHeaderAccessor::MarkInitStateDirty(Header);
 			if (Header.IsBound())
 			{
-				MarkNetObjectStateDirty(Header);
+				MarkNetObjectStateHeaderDirty(Header);
 			}
 		}
 	}
diff --git a/Engine/Source/Runtime/Experimental/Iris/Core/Private/Iris/ReplicationSystem/DirtyNetObjectTracker.cpp b/Engine/Source/Runtime/Experimental/Iris/Core/Private/Iris/ReplicationSystem/DirtyNetObjectTracker.cpp
--- a/Engine/Source/Runtime/Experimental/Iris/Core/Private/Iris/ReplicationSystem/DirtyNetObjectTracker.cpp
+++ b/Engine/Source/Runtime/Experimental/Iris/Core/Private/Iris/ReplicationSystem/DirtyNetObjectTracker.cpp
@@ -70,12 +70,33 @@ void FDirtyNetObjectTracker::Init(const FDirtyNetObjectTrackerInitParams& Params
 void FDirtyNetObjectTracker::Deinit()
 {
 	GlobalDirtyTrackerPollHandle.Destroy();
-	bHasPolledGlobalDirtyTracker = false;
+	bShouldResetPolledGlobalDirtyTracker = false;
 
 	delete[] DirtyNetObjectContainer;
 	DirtyNetObjectContainer = nullptr;
 }
 
+void FDirtyNetObjectTracker::GrabAndApplyGlobalDirtyObjectList()
+{
+	{
+		const TSet<FNetHandle>& GlobalDirtyNetObjects = FGlobalDirtyNetObjectTracker::GetDirtyNetObjects(GlobalDirtyTrackerPollHandle);
+		for (FNetHandle NetHandle : GlobalDirtyNetObjects)
+		{
+			const FInternalNetRefIndex NetObjectIndex = NetRefHandleManager->GetInternalIndexFromNetHandle(NetHandle);
+			if (NetObjectIndex != FNetRefHandleManager::InvalidInternalIndex)
+			{
+				const uint32 BitOffset = NetObjectIndex;
+				const StorageType BitMask = StorageType(1) << (BitOffset & (StorageTypeBitCount - 1));
+				DirtyNetObjectContainer[BitOffset / StorageTypeBitCount] |= BitMask;
+			}
+		}
+	}
+
+	FGlobalDirtyNetObjectTracker::ResetDirtyNetObjectsIfSinglePoller(GlobalDirtyTrackerPollHandle);
+
+	bShouldResetPolledGlobalDirtyTracker = true;
+}
+
 void FDirtyNetObjectTracker::UpdateDirtyNetObjects()
 {
 	if (!GlobalDirtyTrackerPollHandle.IsValid())
@@ -87,19 +108,9 @@ void FDirtyNetObjectTracker::UpdateDirtyNetObjects()
 
 	LockExternalAccess();
 
-	bHasPolledGlobalDirtyTracker = true;
+	GrabAndApplyGlobalDirtyObjectList();
 
-	const TSet<FNetHandle>& GlobalDirtyNetObjects = FGlobalDirtyNetObjectTracker::GetDirtyNetObjects(GlobalDirtyTrackerPollHandle);
-	for (FNetHandle NetHandle : GlobalDirtyNetObjects)
-	{
-		const FInternalNetRefIndex NetObjectIndex = NetRefHandleManager->GetInternalIndexFromNetHandle(NetHandle);
-		if (NetObjectIndex != FNetRefHandleManager::InvalidInternalIndex)
-		{
-			const uint32 BitOffset = NetObjectIndex;
-			const StorageType BitMask = StorageType(1) << (BitOffset & (StorageTypeBitCount - 1));
-			DirtyNetObjectContainer[BitOffset/StorageTypeBitCount] |= BitMask;
-		}
-	}
+	//$IRIS TODO:  We could look if any objects where actually in the global list and skip the array iteration if not needed.
 
 	const uint32* GlobalScopeListData = NetRefHandleManager->GetCurrentFrameScopableInternalIndices().GetData();
 	uint32* AccumulatedDirtyNetObjectsData = AccumulatedDirtyNetObjects.GetData();
@@ -111,13 +122,25 @@ void FDirtyNetObjectTracker::UpdateDirtyNetObjects()
 		uint32 DirtyObjectWord = DirtyNetObjectContainer[WordIndex] & GlobalScopeListData[WordIndex];
 		DirtyNetObjectContainer[WordIndex] = DirtyObjectWord;
 
-		// Add new dirty objects to the accumulated list
-		AccumulatedDirtyNetObjectsData[WordIndex] = AccumulatedDirtyNetObjectsData[WordIndex] | DirtyObjectWord;
+		// Add the latest dirty objects to the accumulated list and remove no-longer scoped objects that have never been copied.
+		AccumulatedDirtyNetObjectsData[WordIndex] = (AccumulatedDirtyNetObjectsData[WordIndex] | DirtyNetObjectContainer[WordIndex]) & GlobalScopeListData[WordIndex];
 	}
 
 	AllowExternalAccess();
 }
 
+void FDirtyNetObjectTracker::UpdateAndLockDirtyNetObjects()
+{
+	if (!GlobalDirtyTrackerPollHandle.IsValid())
+	{
+		return;
+	}
+	
+	UpdateDirtyNetObjects();
+
+	FGlobalDirtyNetObjectTracker::LockDirtyListUntilReset(GlobalDirtyTrackerPollHandle);
+}
+
 void FDirtyNetObjectTracker::UpdateAccumulatedDirtyList()
 {
 	IRIS_PROFILER_SCOPE(FDirtyNetObjectTracker_UpdateDirtyNetObjects)
@@ -181,9 +204,9 @@ void FDirtyNetObjectTracker::ReconcilePolledList(const FNetBitArrayView& Objects
 {
 	LockExternalAccess();
 
-	if (bHasPolledGlobalDirtyTracker)
+	if (bShouldResetPolledGlobalDirtyTracker)
 	{
-		bHasPolledGlobalDirtyTracker = false;
+		bShouldResetPolledGlobalDirtyTracker = false;
 		FGlobalDirtyNetObjectTracker::ResetDirtyNetObjects(GlobalDirtyTrackerPollHandle);
 	}
 
diff --git a/Engine/Source/Runtime/Experimental/Iris/Core/Private/Iris/ReplicationSystem/DirtyNetObjectTracker.h b/Engine/Source/Runtime/Experimental/Iris/Core/Private/Iris/ReplicationSystem/DirtyNetObjectTracker.h
--- a/Engine/Source/Runtime/Experimental/Iris/Core/Private/Iris/ReplicationSystem/DirtyNetObjectTracker.h
+++ b/Engine/Source/Runtime/Experimental/Iris/Core/Private/Iris/ReplicationSystem/DirtyNetObjectTracker.h
@@ -38,9 +38,15 @@ public:
 
 	void Init(const FDirtyNetObjectTrackerInitParams& Params);
 
+	/** Returns true if this dirty tracker can be used by the replication system */
+	bool IsInit() const { return NetRefHandleManager != nullptr; }
+
 	/** Update dirty objects with the set of globally marked dirty objects. */
 	void UpdateDirtyNetObjects();
 
+	/* Update dirty objects from the global list and then prevent future modifications to that list until it is reset. */
+	void UpdateAndLockDirtyNetObjects();
+
 	/** Add all the current frame dirty objects set into the accumulated list */
 	void UpdateAccumulatedDirtyList();
 
@@ -72,6 +78,7 @@ private:
 	void Deinit();
 	void MarkNetObjectDirty(FInternalNetRefIndex NetObjectIndex);
 	void ForceNetUpdate(FInternalNetRefIndex NetObjectIndex);
+	void GrabAndApplyGlobalDirtyObjectList();
 
 	/** Can only be accessed via FDirtyObjectsAccessor */
 	FNetBitArrayView GetDirtyNetObjectsThisFrame();
@@ -98,7 +105,7 @@ private:
 	uint32 NetObjectIdRangeEnd = 0;
 	uint32 NetObjectIdCount = 0;
 	
-	bool bHasPolledGlobalDirtyTracker = false;
+	bool bShouldResetPolledGlobalDirtyTracker = false;
 
 #if UE_NET_THREAD_SAFETY_CHECK
 	std::atomic_bool bIsExternalAccessAllowed = false;
diff --git a/Engine/Source/Runtime/Experimental/Iris/Core/Private/Iris/ReplicationSystem/ObjectReplicationBridge.cpp b/Engine/Source/Runtime/Experimental/Iris/Core/Private/Iris/ReplicationSystem/ObjectReplicationBridge.cpp
--- a/Engine/Source/Runtime/Experimental/Iris/Core/Private/Iris/ReplicationSystem/ObjectReplicationBridge.cpp
+++ b/Engine/Source/Runtime/Experimental/Iris/Core/Private/Iris/ReplicationSystem/ObjectReplicationBridge.cpp
@@ -791,6 +791,8 @@ void UObjectReplicationBridge::PreSendUpdate()
 
 	PreUpdate(ObjectsConsideredForPolling);
 
+	FinalizeDirtyObjects();
+
 	ReconcileNewSubObjects(ObjectsConsideredForPolling);
 
 	Poll(ObjectsConsideredForPolling);
@@ -985,7 +987,7 @@ void UObjectReplicationBridge::BuildPollList(UE::Net::FNetBitArrayView ObjectsCo
 
 			{
 				FDirtyObjectsAccessor DirtyObjectsAccessor(ReplicationSystemInternal->GetDirtyNetObjectTracker());
-				FNetBitArrayView DirtyObjectsThisFrame = DirtyObjectsAccessor.GetDirtyNetObjects();
+				const FNetBitArrayView DirtyObjectsThisFrame = DirtyObjectsAccessor.GetDirtyNetObjects();
 
 				DirtyAndRelevantObjectsView.Set(DirtyObjectsThisFrame, FNetBitArray::OrOp, ForceNetUpdateObjects);
 				DirtyAndRelevantObjectsView.Combine(RelevantObjects, FNetBitArray::AndOp);
@@ -1052,6 +1054,16 @@ void UObjectReplicationBridge::Poll(const UE::Net::FNetBitArrayView ObjectsConsi
 	UE_NET_TRACE_FRAME_STATSCOUNTER(ReplicationSystem->GetId(), ReplicationSystem.PolledReferencesObjectCount, Stats.PolledReferencesObjectCount, ENetTraceVerbosity::Trace);
 }
 
+void UObjectReplicationBridge::FinalizeDirtyObjects()
+{
+	using namespace UE::Net::Private;
+
+	IRIS_PROFILER_SCOPE(UObjectReplicationBridge_FinalizeDirtyObjects);
+
+	// Look for new dirty pushmodel objects and then prevent future modifications to it.
+	GetReplicationSystem()->GetReplicationSystemInternal()->GetDirtyNetObjectTracker().UpdateAndLockDirtyNetObjects();
+}
+
 void UObjectReplicationBridge::ReconcileNewSubObjects(UE::Net::FNetBitArrayView ObjectsConsideredForPolling)
 {
 	using namespace UE::Net;
diff --git a/Engine/Source/Runtime/Experimental/Iris/Core/Private/Iris/ReplicationSystem/Polling/ObjectPoller.cpp b/Engine/Source/Runtime/Experimental/Iris/Core/Private/Iris/ReplicationSystem/Polling/ObjectPoller.cpp
--- a/Engine/Source/Runtime/Experimental/Iris/Core/Private/Iris/ReplicationSystem/Polling/ObjectPoller.cpp
+++ b/Engine/Source/Runtime/Experimental/Iris/Core/Private/Iris/ReplicationSystem/Polling/ObjectPoller.cpp
@@ -104,12 +104,8 @@ void FObjectPoller::PollSingleObject(FNetRefHandle Handle)
 
 		FDirtyObjectsAccessor DirtyObjectsAccessor(ReplicationSystemInternal->GetDirtyNetObjectTracker());
 		DirtyObjectsThisFrame = DirtyObjectsAccessor.GetDirtyNetObjects();
-		// From here we call into user code via PreUpdateInstanceFunction, so allow external code to set dirty flags since DirtyObjects is not read anymore.
-		ReplicationSystemInternal->GetDirtyNetObjectTracker().AllowExternalAccess();
 
-		{
-			ForcePollObject(InternalObjectIndex);
-		}
+		ForcePollObject(InternalObjectIndex);
 
 		// Clear ref to locked dirty bit array
 		DirtyObjectsThisFrame = FNetBitArrayView();
diff --git a/Engine/Source/Runtime/Experimental/Iris/Core/Private/Iris/ReplicationSystem/ReplicationOperations.cpp b/Engine/Source/Runtime/Experimental/Iris/Core/Private/Iris/ReplicationSystem/ReplicationOperations.cpp
--- a/Engine/Source/Runtime/Experimental/Iris/Core/Private/Iris/ReplicationSystem/ReplicationOperations.cpp
+++ b/Engine/Source/Runtime/Experimental/Iris/Core/Private/Iris/ReplicationSystem/ReplicationOperations.cpp
@@ -560,6 +560,7 @@ bool FReplicationInstanceOperations::PollAndRefreshCachedPropertyData(const FRep
 	const uint32 FragmentCount = InstanceProtocol->FragmentCount;
 	for (uint32 StateIt = 0; StateIt < FragmentCount; ++StateIt)
 	{
+		// Only poll fragments with NeedsPoll and none of the ExcludedTraits set.
 		const EReplicationFragmentTraits MaskedFragmentTraits = Fragments[StateIt]->GetTraits() & (EReplicationFragmentTraits::NeedsPoll | ExcludeTraits);
 		if (MaskedFragmentTraits == EReplicationFragmentTraits::NeedsPoll)
 		{
diff --git a/Engine/Source/Runtime/Experimental/Iris/Core/Private/Iris/ReplicationSystem/ReplicationSystem.cpp b/Engine/Source/Runtime/Experimental/Iris/Core/Private/Iris/ReplicationSystem/ReplicationSystem.cpp
--- a/Engine/Source/Runtime/Experimental/Iris/Core/Private/Iris/ReplicationSystem/ReplicationSystem.cpp
+++ b/Engine/Source/Runtime/Experimental/Iris/Core/Private/Iris/ReplicationSystem/ReplicationSystem.cpp
@@ -93,17 +93,19 @@ public:
 
 		const uint32 MaxObjectCount =  NetRefHandleManager.GetMaxActiveObjectCount();
 
-		// $IRIS TODO: Need object ID range. Currently abusing hardcoded values from FNetRefHandleManager
-		FDirtyNetObjectTracker& DirtyNetObjectTracker = ReplicationSystemInternal.GetDirtyNetObjectTracker();
-		FDirtyNetObjectTrackerInitParams DirtyNetObjectTrackerInitParams;
+		// DirtyNetObjectTracking is only needed when object replication is allowed
+		if (Params.bAllowObjectReplication)
 		{
+			// $IRIS TODO: Need object ID range. Currently abusing hardcoded values from FNetRefHandleManager
+			FDirtyNetObjectTrackerInitParams DirtyNetObjectTrackerInitParams;
+
 			DirtyNetObjectTrackerInitParams.NetRefHandleManager = &NetRefHandleManager;
 			DirtyNetObjectTrackerInitParams.ReplicationSystemId = ReplicationSystemId;
 			DirtyNetObjectTrackerInitParams.MaxObjectCount = MaxObjectCount;
 			DirtyNetObjectTrackerInitParams.NetObjectIndexRangeStart = 1;
 			DirtyNetObjectTrackerInitParams.NetObjectIndexRangeEnd = MaxObjectCount - 1U;
 
-			DirtyNetObjectTracker.Init(DirtyNetObjectTrackerInitParams);
+			ReplicationSystemInternal.InitDirtyNetObjectTracker(DirtyNetObjectTrackerInitParams);
 		}
 
 		FReplicationStateStorage& StateStorage = ReplicationSystemInternal.GetReplicationStateStorage();
@@ -779,7 +781,10 @@ void UReplicationSystem::PostSendUpdate()
 	// Most systems are only updated during the normal TickFlush
 	if (Impl->CurrentSendPass == EReplicationSystemSendPass::TickFlush)
 	{
-		Impl->ResetObjectStateDirtiness();
+		if (bAllowObjectReplication)
+		{
+			Impl->ResetObjectStateDirtiness();
+		}
 
 		Impl->EndPostSendUpdate();
 
diff --git a/Engine/Source/Runtime/Experimental/Iris/Core/Private/Iris/ReplicationSystem/ReplicationSystemInternal.h b/Engine/Source/Runtime/Experimental/Iris/Core/Private/Iris/ReplicationSystem/ReplicationSystemInternal.h
--- a/Engine/Source/Runtime/Experimental/Iris/Core/Private/Iris/ReplicationSystem/ReplicationSystemInternal.h
+++ b/Engine/Source/Runtime/Experimental/Iris/Core/Private/Iris/ReplicationSystem/ReplicationSystemInternal.h
@@ -54,7 +54,8 @@ public:
 	FNetRefHandleManager& GetNetRefHandleManager() { return NetRefHandleManager; }
 	const FNetRefHandleManager& GetNetRefHandleManager() const { return NetRefHandleManager; }
 
-	FDirtyNetObjectTracker& GetDirtyNetObjectTracker() { return DirtyNetObjectTracker; }
+	void InitDirtyNetObjectTracker(const struct FDirtyNetObjectTrackerInitParams& Params) { DirtyNetObjectTracker.Init(Params); }
+	FDirtyNetObjectTracker& GetDirtyNetObjectTracker() { checkf(DirtyNetObjectTracker.IsInit(), TEXT("Not allowed to access the DirtyNetObjectTracker unless object replication is enabled.")); return DirtyNetObjectTracker; }
 
 	FReplicationStateDescriptorRegistry& GetReplicationStateDescriptorRegistry() { return ReplicationStateDescriptorRegistry; }
 
diff --git a/Engine/Source/Runtime/Experimental/Iris/Core/Public/Iris/ReplicationState/ReplicationStateUtil.h b/Engine/Source/Runtime/Experimental/Iris/Core/Public/Iris/ReplicationState/ReplicationStateUtil.h
--- a/Engine/Source/Runtime/Experimental/Iris/Core/Public/Iris/ReplicationState/ReplicationStateUtil.h
+++ b/Engine/Source/Runtime/Experimental/Iris/Core/Public/Iris/ReplicationState/ReplicationStateUtil.h
@@ -9,21 +9,15 @@
 #include "Iris/ReplicationState/ReplicationStateDescriptor.h"
 #include "Iris/ReplicationState/ReplicationStateFwd.h"
 
-namespace UE::Net::Private
-{
-	extern IRISCORE_API void MarkNetObjectStateDirty(uint32 ReplicationSystemId, uint32 NetObjectIndex);
-}
-
 namespace UE::Net
 {
-/** Mark the a NetObject/Handle as dirty which means it requires to be copied when we update the Replication system */
-inline void MarkNetObjectStateDirty(FReplicationStateHeader& Header)
-{
-	MarkNetObjectStateDirty(FNetHandleManager::MakeNetHandleFromId(Private::FReplicationStateHeaderAccessor::GetNetHandleId(Header)));
-	Private::FReplicationStateHeaderAccessor::MarkStateDirty(Header);
-}
+	/** Mark the Replication State Header as dirty so it can be copied in the CopyDirtyState pass. */
+	inline void MarkNetObjectStateHeaderDirty(FReplicationStateHeader& Header)
+	{
+		Private::FReplicationStateHeaderAccessor::MarkStateDirty(Header);
+	}
 
-}
+} // end namespace UE::Net
 
 namespace UE::Net::Private
 {
@@ -71,7 +65,7 @@ inline void MarkDirty(UE::Net::FReplicationStateHeader& InternalState, FNetBitAr
 	if (InternalState.IsBound() && !MemberChangeMask.GetBit(ChangeMaskInfo.BitOffset))
 	{
 		FReplicationStateHeaderAccessor::MarkStateDirty(InternalState);
-		MarkNetObjectStateDirty(InternalState);
+		MarkNetObjectStateHeaderDirty(InternalState);
 	}
 
 	MemberChangeMask.SetBits(ChangeMaskInfo.BitOffset, ChangeMaskInfo.BitCount);
diff --git a/Engine/Source/Runtime/Experimental/Iris/Core/Public/Iris/ReplicationSystem/FastArrayReplicationFragment.h b/Engine/Source/Runtime/Experimental/Iris/Core/Public/Iris/ReplicationSystem/FastArrayReplicationFragment.h
--- a/Engine/Source/Runtime/Experimental/Iris/Core/Public/Iris/ReplicationSystem/FastArrayReplicationFragment.h
+++ b/Engine/Source/Runtime/Experimental/Iris/Core/Public/Iris/ReplicationSystem/FastArrayReplicationFragment.h
@@ -9,8 +9,11 @@
 #include "Iris/ReplicationState/PropertyReplicationState.h"
 #include "Iris/ReplicationState/ReplicationStateUtil.h"
 #include "Iris/ReplicationState/ReplicationStateDescriptor.h"
+
 #include "Net/Core/NetBitArray.h"
 
+#include "Templates/UnrealTemplate.h"
+
 namespace UE::Net
 {
 
@@ -184,32 +187,40 @@ bool TFastArrayReplicationFragment<FastArrayItemType, FastArrayType>::PollAllSta
 	FastArrayItemType* DstItems = DstWrappedArray->GetData();
 	FastArrayItemType* SrcItems = SrcWrappedArray->GetData();
 
-	// Iterate over array entries and copy the statedata using internal data if it has changed
-	for (int32 ElementIt = 0, ElementEndIt = ElementCount; ElementIt < ElementEndIt; ++ElementIt)
 	{
-		FastArrayItemType& SrcItem = SrcItems[ElementIt];
-		FastArrayItemType& DstItem = DstItems[ElementIt];
+#if WITH_PUSH_MODEL
+		// Disable push model by temporarily setting the FastArray's RepIndex to none.
+		// This prevents the array from adding itself to the global dirty list via MarkItemDirty while we are polling it.
+		TGuardValue DisablePushModel(SrcArraySerializer->RepIndex, (int32)INDEX_NONE);
+#endif
 
-		const bool bIsWritingOnClient = false;
-		if (SrcArraySerializer->template ShouldWriteFastArrayItem<FastArrayItemType, FastArrayType>(SrcItem, bIsWritingOnClient))
+		// Iterate over array entries and copy the statedata using internal data if it has changed
+		for (int32 ElementIt = 0, ElementEndIt = ElementCount; ElementIt < ElementEndIt; ++ElementIt)
 		{
-			if (SrcItem.ReplicationID == INDEX_NONE)
-			{
-				SrcArraySerializer->MarkItemDirty(SrcItem);
-			}
+			FastArrayItemType& SrcItem = SrcItems[ElementIt];
+			FastArrayItemType& DstItem = DstItems[ElementIt];
 
-			const bool bReplicationKeyChanged = SrcItem.ReplicationKey != DstItem.ReplicationKey || SrcItem.ReplicationID != DstItem.ReplicationID;
-			if (bReplicationKeyChanged || (bForceFullCompare && !InternalCompareArrayElement(ArrayElementDescriptor, &DstItem, &SrcItem)))
+			const bool bIsWritingOnClient = false;
+			if (SrcArraySerializer->template ShouldWriteFastArrayItem<FastArrayItemType, FastArrayType>(SrcItem, bIsWritingOnClient))
 			{
-				InternalCopyArrayElement(ArrayElementDescriptor, &DstItem, &SrcItem);
-				DstItem.ReplicationKey = SrcItem.ReplicationKey;
+				if (SrcItem.ReplicationID == INDEX_NONE)
+				{
+					SrcArraySerializer->MarkItemDirty(SrcItem);
+				}
 
-				// Mark element as dirty and mark array as dirty as well.
-				if (ChangeMaskBitCount)
+				const bool bReplicationKeyChanged = SrcItem.ReplicationKey != DstItem.ReplicationKey || SrcItem.ReplicationID != DstItem.ReplicationID;
+				if (bReplicationKeyChanged || (bForceFullCompare && !InternalCompareArrayElement(ArrayElementDescriptor, &DstItem, &SrcItem)))
 				{
-					MemberChangeMask.SetBit((ElementIt % ChangeMaskBitCount) + ChangeMaskBitOffset);
+					InternalCopyArrayElement(ArrayElementDescriptor, &DstItem, &SrcItem);
+					DstItem.ReplicationKey = SrcItem.ReplicationKey;
+
+					// Mark element as dirty and mark array as dirty as well.
+					if (ChangeMaskBitCount)
+					{
+						MemberChangeMask.SetBit((ElementIt % ChangeMaskBitCount) + ChangeMaskBitOffset);
+					}
+					bMarkArrayDirty = true;
 				}
-				bMarkArrayDirty = true;
 			}
 		}
 	}
@@ -221,7 +232,7 @@ bool TFastArrayReplicationFragment<FastArrayItemType, FastArrayType>::PollAllSta
 	if (bMarkArrayDirty && ReplicationState->IsCustomConditionEnabled(FIrisFastArraySerializer::IrisFastArrayPropertyBitIndex))
 	{
 		MemberChangeMask.SetBit(FIrisFastArraySerializer::IrisFastArrayPropertyBitIndex);
-		MarkNetObjectStateDirty(UE::Net::Private::GetReplicationStateHeader(ReplicationState->GetStateBuffer(), ReplicationStateDescriptor));
+		MarkNetObjectStateHeaderDirty(UE::Net::Private::GetReplicationStateHeader(ReplicationState->GetStateBuffer(), ReplicationStateDescriptor));
 	}
 
 	return IsDirty();
@@ -242,7 +253,7 @@ void TFastArrayReplicationFragment<FastArrayItemType, FastArrayType>::MarkDirty(
 	{
 		FNetBitArrayView MemberChangeMask = UE::Net::Private::GetMemberChangeMask(ReplicationState->GetStateBuffer(), ReplicationStateDescriptor);
 		MemberChangeMask.SetBit(FIrisFastArraySerializer::IrisFastArrayPropertyBitIndex);
-		MarkNetObjectStateDirty(UE::Net::Private::GetReplicationStateHeader(ReplicationState->GetStateBuffer(), ReplicationStateDescriptor));
+		MarkNetObjectStateHeaderDirty(UE::Net::Private::GetReplicationStateHeader(ReplicationState->GetStateBuffer(), ReplicationStateDescriptor));
 	}
 }
 
diff --git a/Engine/Source/Runtime/Experimental/Iris/Core/Public/Iris/ReplicationSystem/ObjectReplicationBridge.h b/Engine/Source/Runtime/Experimental/Iris/Core/Public/Iris/ReplicationSystem/ObjectReplicationBridge.h
--- a/Engine/Source/Runtime/Experimental/Iris/Core/Public/Iris/ReplicationSystem/ObjectReplicationBridge.h
+++ b/Engine/Source/Runtime/Experimental/Iris/Core/Public/Iris/ReplicationSystem/ObjectReplicationBridge.h
@@ -318,6 +318,9 @@ private:
 	/** Call the user function PreUpdate (aka PreReplication) on objects about to be polled. */
 	void PreUpdate(const UE::Net::FNetBitArrayView ObjectsConsideredForPolling);
 
+	/** Find any objects that got set dirty by user code during PreUpdate then lock future modifications to the global dirty list */
+	void FinalizeDirtyObjects();
+
 	/** Find any new subobjects created inside PreUpdate and ensure they will be replicated this frame */
 	void ReconcileNewSubObjects(UE::Net::FNetBitArrayView ObjectsConsideredForPolling);
 
diff --git a/Engine/Source/Runtime/Net/Core/Private/Net/Core/DirtyNetObjectTracker/GlobalDirtyNetObjectTracker.cpp b/Engine/Source/Runtime/Net/Core/Private/Net/Core/DirtyNetObjectTracker/GlobalDirtyNetObjectTracker.cpp
--- a/Engine/Source/Runtime/Net/Core/Private/Net/Core/DirtyNetObjectTracker/GlobalDirtyNetObjectTracker.cpp
+++ b/Engine/Source/Runtime/Net/Core/Private/Net/Core/DirtyNetObjectTracker/GlobalDirtyNetObjectTracker.cpp
@@ -20,10 +20,37 @@ private:
 
 	static TSet<FNetHandle> EmptyDirtyObjects;
 
+private:
+
+	struct FPollerStatus
+	{
+		// Is this status tied to an active registered poller
+		uint8 bIsActive:1;
+
+		// Has this poller read the dirty list this frame yet.
+		uint8 bHasGathered:1;
+
+		FPollerStatus() : bIsActive(false), bHasGathered(false) {}
+
+		void ClearStatus()
+		{
+			*this = FPollerStatus();
+		}
+	};
+
+private:
+
 	TSet<FNetHandle> DirtyObjects;
+
 	FNetBitArray AssignedHandleIndices;
 	FNetBitArray Pollers;
+
+	TArray<FPollerStatus> PollerStatuses;
+
 	uint32 PollerCount = 0;
+
+	/** When true detect and prevent illegal changes to the dirty object list. */
+	bool bLockDirtyList = false;
 };
 
 TSet<FNetHandle> FGlobalDirtyNetObjectTracker::FPimpl::EmptyDirtyObjects;
@@ -34,18 +61,11 @@ void FGlobalDirtyNetObjectTracker::MarkNetObjectStateDirty(FNetHandle NetHandle)
 {
 	if (Instance && Instance->PollerCount > 0)
 	{
-		Instance->DirtyObjects.Add(NetHandle);
-	}
-}
-
-bool FGlobalDirtyNetObjectTracker::IsNetObjectStateDirty(FNetHandle NetHandle)
-{
-	if (Instance)
-	{
-		return Instance->DirtyObjects.Find(NetHandle) != nullptr;
+		if (ensureMsgf(!Instance->bLockDirtyList, TEXT("MarkNetObjectStateDirty was called while the dirty list was set to read-only.")))
+		{
+			Instance->DirtyObjects.Add(NetHandle);
+		}
 	}
-
-	return false;
 }
 
 FGlobalDirtyNetObjectTracker::FPollHandle FGlobalDirtyNetObjectTracker::CreatePoller()
@@ -66,6 +86,10 @@ FGlobalDirtyNetObjectTracker::FPollHandle FGlobalDirtyNetObjectTracker::CreatePo
 
 		Instance->AssignedHandleIndices.SetBit(HandleIndex);
 		++Instance->PollerCount;
+
+		Instance->PollerStatuses.SetNum(Instance->PollerCount, false);
+		Instance->PollerStatuses[HandleIndex].bIsActive = true;
+
 		return FPollHandle(HandleIndex);
 	}
 
@@ -79,18 +103,21 @@ void FGlobalDirtyNetObjectTracker::DestroyPoller(uint32 HandleIndex)
 		return;
 	}
 
-	if (ensureAlwaysMsgf((HandleIndex < Instance->AssignedHandleIndices.GetNumBits()) && Instance->AssignedHandleIndices.GetBit(HandleIndex), TEXT("Destroying unknown poller with handle index %u"), HandleIndex))
+	if (ensureMsgf((HandleIndex < Instance->AssignedHandleIndices.GetNumBits()) && Instance->AssignedHandleIndices.GetBit(HandleIndex), TEXT("Destroying unknown poller with handle index %u"), HandleIndex))
 	{
 		Instance->AssignedHandleIndices.ClearBit(HandleIndex);
 
 		const uint32 PollerCalled = Instance->Pollers.GetBit(HandleIndex);
-		ensureAlwaysMsgf(PollerCalled == 0U, TEXT("Destroying poller that called GetDirtyNetObjects() but not ResetDirtyNetObjects()"));
+		ensureMsgf(PollerCalled == 0U, TEXT("Destroying poller that called GetDirtyNetObjects() but not ResetDirtyNetObjects()"));
 		Instance->Pollers.ClearBit(HandleIndex);
 
+		Instance->PollerStatuses[HandleIndex].ClearStatus();
+
 		--Instance->PollerCount;
 		if (Instance->PollerCount <= 0)
 		{
 			Instance->DirtyObjects.Reset();
+			Instance->bLockDirtyList = false;
 		}
 	}
 }
@@ -99,26 +126,74 @@ const TSet<FNetHandle>& FGlobalDirtyNetObjectTracker::GetDirtyNetObjects(const F
 {
 	if (Instance && Handle.IsValid())
 	{
-		//$IRIS TODO: Once a poller reads the dirty list, we should make this list immutable and prevent any changes to it.
-		//			  Otherwise any changes between the first Get and the final Reset will get discarded and never seen by all the Pollers.
+		check(Instance->PollerStatuses[Handle.Index].bIsActive);
+
 		Instance->Pollers.SetBit(Handle.Index);
+		Instance->PollerStatuses[Handle.Index].bHasGathered = true;
 		return Instance->DirtyObjects;
 	}
 
 	return FPimpl::EmptyDirtyObjects;
 }
 
+void FGlobalDirtyNetObjectTracker::LockDirtyListUntilReset(const FPollHandle& Handle)
+{
+	if (Instance && Handle.IsValid())
+	{
+		check(Instance->PollerStatuses[Handle.Index].bIsActive);
+
+		// From here prevent new dirty objects until the list is reset.
+		Instance->bLockDirtyList = true;
+	}
+}
+
 void FGlobalDirtyNetObjectTracker::ResetDirtyNetObjects(const FPollHandle& Handle)
 {
 	if (Instance && Handle.IsValid())
 	{
+		check(Instance->PollerStatuses[Handle.Index].bIsActive);
+
 		Instance->Pollers.ClearBit(Handle.Index);
-		//$IRIS TODO: This assumes every Poller called Get before anyone called Reset. We should make sure this is the case all the time.
+
 		if (Instance->Pollers.IsNoBitSet())
 		{
+
+#if DO_CHECK
+			bool bAllPollersGathered = true;
+			for (const FPimpl::FPollerStatus& PollerStatus : Instance->PollerStatuses)
+			{
+				bAllPollersGathered &= !PollerStatus.bIsActive || PollerStatus.bHasGathered;
+			}
+			ensureMsgf(bAllPollersGathered, TEXT("Not all pollers gathered the dirty list before the list got reset. Those pollers will never know those objects were dirty."));
+#endif
+
+			Instance->bLockDirtyList = false;
+			Instance->DirtyObjects.Reset();
+		}
+	}
+}
+
+bool FGlobalDirtyNetObjectTracker::ResetDirtyNetObjectsIfSinglePoller(const FPollHandle& Handle)
+{
+	if (Instance && Handle.IsValid())
+	{
+		check(Instance->PollerStatuses[Handle.Index].bIsActive);
+
+		// Are we the only poller registered to read and reset the list
+		if (Instance->AssignedHandleIndices.CountSetBits() == 1)
+		{
+			check(Instance->Pollers.IsBitSet(Handle.Index));
+
+			Instance->bLockDirtyList = false;
 			Instance->DirtyObjects.Reset();
+
+			Instance->Pollers.ClearBit(Handle.Index);
+
+			return true;
 		}
 	}
+
+	return false;
 }
 
 void FGlobalDirtyNetObjectTracker::Init()
@@ -155,6 +230,8 @@ void FGlobalDirtyNetObjectTracker::FPimpl::ResetDirtyNetObjects()
 
 		// DirtyObjects should already be reset if there are no pollers.
 		DirtyObjects.Reset();
+
+		bLockDirtyList = false;
 	}
 }
 
diff --git a/Engine/Source/Runtime/Net/Core/Public/Net/Core/DirtyNetObjectTracker/GlobalDirtyNetObjectTracker.h b/Engine/Source/Runtime/Net/Core/Public/Net/Core/DirtyNetObjectTracker/GlobalDirtyNetObjectTracker.h
--- a/Engine/Source/Runtime/Net/Core/Public/Net/Core/DirtyNetObjectTracker/GlobalDirtyNetObjectTracker.h
+++ b/Engine/Source/Runtime/Net/Core/Public/Net/Core/DirtyNetObjectTracker/GlobalDirtyNetObjectTracker.h
@@ -50,9 +50,6 @@ public:
 	/** Marks an object as dirty. Assumes Init() has been called. */
 	NETCORE_API static void MarkNetObjectStateDirty(FNetHandle);
 
-	/** Returns true if an object has been marked as dirty since the last reset. Assumes Init() has been called. */
-	NETCORE_API static bool IsNetObjectStateDirty(FNetHandle);
-
 	/** Create a poller which is assumed to call PollDirtyNetObjects(). Assumes Init() has been called. */
 	NETCORE_API static FPollHandle CreatePoller();
 
@@ -61,7 +58,12 @@ public:
 	 * A matching call to ResetDirtyNetObjects() is required but should not be called until all pollers have had the chance to call GetDirtyNetObjects().
 	 * Assumes Init() has been called.
 	 */
-	NETCORE_API static const TSet<FNetHandle>& GetDirtyNetObjects(const FPollHandle&);
+	NETCORE_API static const TSet<FNetHandle>& GetDirtyNetObjects(const FPollHandle& Handle);
+
+	/**
+	 * Detect any updates to the dirty list until all pollers have gathered and reset the list.
+	 */
+	NETCORE_API static void LockDirtyListUntilReset(const FPollHandle& Handle);
 
 	/**
 	 * Clears the set of dirty net objects. GetDirtyNetObjects() must be called first. Some synchronization is needed between pollers such that all pollers have the chance to call GetDirtyNetObjects() first.
@@ -70,6 +72,12 @@ public:
 	 */
 	NETCORE_API static void ResetDirtyNetObjects(const FPollHandle&);
 
+	/**
+	 * Reset the list of dirty net objects but only if there is a single poller registered in the system.
+	 * Returns true if the reset was executed.
+	 */
+	NETCORE_API static bool ResetDirtyNetObjectsIfSinglePoller(const FPollHandle&);
+
 protected:
 	friend FNetCoreModule;
 
