#ue5 - UCharacterMovementComponent::RoundAcceleration() changed to match the same rounding and clamping that is used by replication of FVector_NetQuantize10. Adapted from PR #10236

#tests PIE with networked movement showing corrections/logging

diff --git a/Engine/Source/Runtime/Engine/Private/Components/CharacterMovementComponent.cpp b/Engine/Source/Runtime/Engine/Private/Components/CharacterMovementComponent.cpp
--- a/Engine/Source/Runtime/Engine/Private/Components/CharacterMovementComponent.cpp
+++ b/Engine/Source/Runtime/Engine/Private/Components/CharacterMovementComponent.cpp
@@ -7712,11 +7712,8 @@ FVector UCharacterMovementComponent::ScaleInputAcceleration(const FVector& Input
 
 FVector UCharacterMovementComponent::RoundAcceleration(FVector InAccel) const
 {
-	// Match FVector_NetQuantize10 (1 decimal place of precision).
-	InAccel.X = FMath::RoundToFloat(InAccel.X * 10.f) / 10.f;
-	InAccel.Y = FMath::RoundToFloat(InAccel.Y * 10.f) / 10.f;
-	InAccel.Z = FMath::RoundToFloat(InAccel.Z * 10.f) / 10.f;
-	return InAccel;
+	// Match FVector_NetQuantize10
+	return UE::Net::QuantizeVector(10, InAccel);
 }
 
 float UCharacterMovementComponent::ComputeAnalogInputModifier() const
diff --git a/Engine/Source/Runtime/Net/Core/Private/Net/Core/Serialization/QuantizedVectorSerialization.cpp b/Engine/Source/Runtime/Net/Core/Private/Net/Core/Serialization/QuantizedVectorSerialization.cpp
--- a/Engine/Source/Runtime/Net/Core/Private/Net/Core/Serialization/QuantizedVectorSerialization.cpp
+++ b/Engine/Source/Runtime/Net/Core/Private/Net/Core/Serialization/QuantizedVectorSerialization.cpp
@@ -187,6 +187,79 @@ bool ReadQuantizedVector(const int32 Scale, T& Value, FArchive& Ar)
 	}
 }
 
+template<class T>
+T QuantizeVector(const int32 Scale, const T& Value)
+{
+	using ScalarType = decltype(T::X);
+	constexpr SIZE_T ScalarTypeSize = sizeof(ScalarType);
+	using IntType = typename TSignedIntType<ScalarTypeSize>::Type;
+
+	static_assert(ScalarTypeSize == 4U || ScalarTypeSize == 8U, "Unknown floating point type.");
+
+	// Beyond 2^MaxExponentForScaling scaling cannot improve the precision as the next floating point value is at least 1.0 more. 
+	constexpr uint32 MaxExponentForScaling = ScalarTypeSize == 4 ? 23U : 52U;
+	constexpr ScalarType MaxValueToScale = ScalarType(IntType(1) << MaxExponentForScaling);
+
+	// Rounding of large values can introduce additional precision errors and the extra cost to serialize with full precision is small.
+	constexpr uint32 MaxExponentAfterScaling = ScalarTypeSize == 4 ? 30U : 62U;
+	constexpr ScalarType MaxScaledValue = ScalarType(IntType(1) << MaxExponentAfterScaling);
+
+	// NaN values can be properly serialized using the full precision path, but they typically cause lots of errors
+	// for the typical engine use case.
+	if (Value.ContainsNaN())
+	{
+		logOrEnsureNanError(TEXT("%s"), TEXT("QuantizeVector: Value isn't finite. Clearing for safety."));
+		return T{ 0,0,0 };
+	}
+
+	const ScalarType Factor = IntCastChecked<int16>(Scale);
+	T ScaledValue;
+	ScaledValue.X = Value.X * Factor;
+	ScaledValue.Y = Value.Y * Factor;
+	ScaledValue.Z = Value.Z * Factor;
+
+	// If the component values are within bounds then we optimize the bandwidth, otherwise we use full precision.
+	if (ScaledValue.GetAbsMax() < MaxScaledValue)
+	{
+		const bool bUseScaledValue = Value.GetAbsMin() < MaxValueToScale;
+
+		// 'Write' value
+		IntType X;
+		IntType Y;
+		IntType Z;
+		if (bUseScaledValue)
+		{
+			X = RoundFloatToInt(ScaledValue.X);
+			Y = RoundFloatToInt(ScaledValue.Y);
+			Z = RoundFloatToInt(ScaledValue.Z);
+		}
+		else
+		{
+			X = RoundFloatToInt(Value.X);
+			Y = RoundFloatToInt(Value.Y);
+			Z = RoundFloatToInt(Value.Z);
+		}
+
+		// 'Read' value
+		T TempValue;
+		TempValue.X = ScalarType(X);
+		TempValue.Y = ScalarType(Y);
+		TempValue.Z = ScalarType(Z);
+
+		// Apply scaling if needed.
+		if (bUseScaledValue)
+		{
+			return TempValue / ScalarType(Scale);
+		}
+		else
+		{
+			return TempValue;
+		}
+	}
+
+	return Value;
+}
+
 }
 
 namespace UE::Net
@@ -202,6 +275,11 @@ bool ReadQuantizedVector(int32 Scale, FVector3d& Value, FArchive& Ar)
 	return Private::ReadQuantizedVector(Scale, Value, Ar);
 }
 
+FVector3d QuantizeVector(const int32 Scale, const FVector3d& Value)
+{
+	return Private::QuantizeVector(Scale, Value);
+}
+
 bool WriteQuantizedVector(int32 Scale, const FVector3f& Value, FArchive& Ar)
 {
 	return Private::WriteQuantizedVector(Scale, Value, Ar);
@@ -212,4 +290,9 @@ bool ReadQuantizedVector(int32 Scale, FVector3f& Value, FArchive& Ar)
 	return Private::ReadQuantizedVector(Scale, Value, Ar);
 }
 
+FVector3f QuantizeVector(const int32 Scale, const FVector3f& Value)
+{
+	return Private::QuantizeVector(Scale, Value);
+}
+
 }
diff --git a/Engine/Source/Runtime/Net/Core/Public/Net/Core/Serialization/QuantizedVectorSerialization.h b/Engine/Source/Runtime/Net/Core/Public/Net/Core/Serialization/QuantizedVectorSerialization.h
--- a/Engine/Source/Runtime/Net/Core/Public/Net/Core/Serialization/QuantizedVectorSerialization.h
+++ b/Engine/Source/Runtime/Net/Core/Public/Net/Core/Serialization/QuantizedVectorSerialization.h
@@ -23,6 +23,10 @@ NETCORE_API bool WriteQuantizedVector(const int32 Scale, const FVector3f& Value,
 NETCORE_API bool ReadQuantizedVector(const int32 Scale, FVector3d& Value, FArchive& Ar);
 NETCORE_API bool ReadQuantizedVector(const int32 Scale, FVector3f& Value, FArchive& Ar);
 
+/* Quantize a vector using the same quantization as WriteQuantizedVector followed by ReadQuantizedVector. */
+NETCORE_API FVector3d QuantizeVector(const int32 Scale, const FVector3d& Value);
+NETCORE_API FVector3f QuantizeVector(const int32 Scale, const FVector3f& Value);
+
 template<int32 Scale>
 bool SerializeQuantizedVector(FVector& Value, FArchive& Ar)
 {
