Disabled deprecation of property until I have figured out to compilation errors from 31718958

diff --git a/Engine/Source/Runtime/Net/Core/Classes/Net/Serialization/FastArraySerializer.h b/Engine/Source/Runtime/Net/Core/Classes/Net/Serialization/FastArraySerializer.h
--- a/Engine/Source/Runtime/Net/Core/Classes/Net/Serialization/FastArraySerializer.h
+++ b/Engine/Source/Runtime/Net/Core/Classes/Net/Serialization/FastArraySerializer.h
@@ -520,7 +520,7 @@ struct FFastArraySerializer
 	{
 		// This is the size that the array had before the receive.
 		int32 OldArraySize;
-		UE_DEPRECATED(5.4, "This is unsafe to use and will be removed.")
+		//UE_DEPRECATED(5.4, "This is unsafe to use and will be removed.")
 		uint32 bHasMoreUnmappedReferences : 1U;
 	};
 
