Added console command to set nettrace verbosity at runtime
- NetTrace.SetTraceVerbosity [verbositylevel]

diff --git a/Engine/Source/Runtime/Net/Core/Private/Net/Core/Trace/NetTraceInternal.cpp b/Engine/Source/Runtime/Net/Core/Private/Net/Core/Trace/NetTraceInternal.cpp
--- a/Engine/Source/Runtime/Net/Core/Private/Net/Core/Trace/NetTraceInternal.cpp
+++ b/Engine/Source/Runtime/Net/Core/Private/Net/Core/Trace/NetTraceInternal.cpp
@@ -5,6 +5,7 @@
 #include "Net/Core/Trace/NetDebugName.h"
 #include "Net/Core/Trace/Reporters/NetTraceReporter.h"
 #include "Hash/CityHash.h"
+#include "HAL/IConsoleManager.h"
 #include "Trace/Trace.h"
 #include "UObject/NameTypes.h"
 #include <atomic>
@@ -591,4 +592,20 @@ UE::Net::FNetDebugNameId FNetTrace::TraceName(const UE::Net::FNetDebugName* Debu
 	}
 }
 
+static FAutoConsoleCommand NeTraceSetVerbosityCmd = FAutoConsoleCommand(
+	TEXT("NetTrace.SetTraceVerbosity"),
+	TEXT("Start NetTrace with given verbositylevel."),
+	FConsoleCommandWithArgsDelegate::CreateStatic(
+		[](const TArray<FString>& Args)
+		{
+			if (Args.Num() == 0)
+			{
+				return;
+			}
+
+			FNetTrace::SetTraceVerbosity(FCString::Atoi(*Args[0]));
+		}
+		)
+	);
+
 #endif
