Fix missing export dll tag to SetAnalyticsAppender

#rnx

diff --git a/Engine/Source/Runtime/Net/Core/Public/Net/Core/Analytics/NetAnalytics.h b/Engine/Source/Runtime/Net/Core/Public/Net/Core/Analytics/NetAnalytics.h
--- a/Engine/Source/Runtime/Net/Core/Public/Net/Core/Analytics/NetAnalytics.h
+++ b/Engine/Source/Runtime/Net/Core/Public/Net/Core/Analytics/NetAnalytics.h
@@ -383,7 +383,7 @@ public:
 public:
 
 	using FAnalyticsAppender = TFunction<void(TArray<FAnalyticsEventAttribute>&)>;	
-	void SetAnalyticsAppender(FAnalyticsAppender AppenderFunction);
+	NETCORE_API void SetAnalyticsAppender(FAnalyticsAppender AppenderFunction);
 
 private:
 	/** The analytics provider we are aggregating data for */
