Networking:
* Added Networking_FastArray_Apply csv scope to FastArray reception codepath
* Allows to track regressions in FastArray user code
* Exposed CSV Networking category via NetCore module export

diff --git a/Engine/Source/Runtime/Engine/Private/NetDriver.cpp b/Engine/Source/Runtime/Engine/Private/NetDriver.cpp
--- a/Engine/Source/Runtime/Engine/Private/NetDriver.cpp
+++ b/Engine/Source/Runtime/Engine/Private/NetDriver.cpp
@@ -164,7 +164,7 @@ DECLARE_DWORD_COUNTER_STAT(TEXT("SharedSerialization Property Miss"), STAT_Share
 
 DEFINE_LOG_CATEGORY_STATIC(LogNetSyncLoads, Log, All);
 
-CSV_DEFINE_CATEGORY(Networking, true);
+CSV_DECLARE_CATEGORY_MODULE_EXTERN(NETCORE_API, Networking);
 
 int32 GNumSaturatedConnections; // Counter for how many connections are skipped/early out due to bandwidth saturation
 int32 GNumSharedSerializationHit;
diff --git a/Engine/Source/Runtime/Experimental/Iris/Core/Public/Iris/ReplicationSystem/Private/FastArrayReplicationFragmentInternal.h b/Engine/Source/Runtime/Experimental/Iris/Core/Public/Iris/ReplicationSystem/Private/FastArrayReplicationFragmentInternal.h
--- a/Engine/Source/Runtime/Experimental/Iris/Core/Public/Iris/ReplicationSystem/Private/FastArrayReplicationFragmentInternal.h
+++ b/Engine/Source/Runtime/Experimental/Iris/Core/Public/Iris/ReplicationSystem/Private/FastArrayReplicationFragmentInternal.h
@@ -10,6 +10,9 @@
 #include "Iris/ReplicationState/ReplicationStateDescriptor.h"
 #include "Net/Core/Trace/NetDebugName.h"
 #include "Net/Core/NetBitArray.h"
+#include "ProfilingDebugging/CsvProfiler.h"
+
+CSV_DECLARE_CATEGORY_MODULE_EXTERN(NETCORE_API, Networking);
 
 namespace UE::Net
 {
@@ -203,6 +206,8 @@ void FFastArrayReplicationFragmentHelper::ApplyReplicatedState(FastArrayType* Ds
 {
 	typedef typename ItemArrayType::ElementType ItemType;
 
+	CSV_SCOPED_TIMING_STAT(Networking, FastArray_Apply);
+
 	UE_LOG(LogNetFastTArray, Log, TEXT("FFastArrayReplicationFragmentHelper::ApplyReplicatedState for %s"), Context.Descriptor->DebugName->Name);
 
 	const uint32* ChangeMaskData = Context.StateBufferData.ChangeMaskData;
diff --git a/Engine/Source/Runtime/Net/Core/Classes/Net/Serialization/FastArraySerializer.h b/Engine/Source/Runtime/Net/Core/Classes/Net/Serialization/FastArraySerializer.h
--- a/Engine/Source/Runtime/Net/Core/Classes/Net/Serialization/FastArraySerializer.h
+++ b/Engine/Source/Runtime/Net/Core/Classes/Net/Serialization/FastArraySerializer.h
@@ -20,6 +20,8 @@
 #include "Net/Core/NetCoreModule.h"
 #include "HAL/IConsoleManager.h"
 #include "Templates/EnableIf.h"
+#include "ProfilingDebugging/CsvProfiler.h"
+
 #include "FastArraySerializer.generated.h"
 
 class Error;
@@ -32,6 +34,8 @@ DECLARE_CYCLE_STAT_EXTERN(TEXT("NetSerializeFast Array"), STAT_NetSerializeFastA
 DECLARE_CYCLE_STAT_EXTERN(TEXT("NetSerializeFast Array BuildMap"), STAT_NetSerializeFastArray_BuildMap, STATGROUP_ServerCPU, NETCORE_API);
 DECLARE_CYCLE_STAT_EXTERN(TEXT("NetSerializeFast Array Delta Struct"), STAT_NetSerializeFastArray_DeltaStruct, STATGROUP_ServerCPU, NETCORE_API);
 
+CSV_DECLARE_CATEGORY_MODULE_EXTERN(NETCORE_API, Networking);
+
 /**
  *	===================== Fast TArray Replication ===================== 
  *
@@ -1083,6 +1087,8 @@ void FFastArraySerializer::TFastArraySerializeHelper<Type, SerializerType>::Post
 	TArray<int32, TInlineAllocator<8>>& AddedIndices,
 	GuidMapType& GuidMap)
 {
+	CSV_SCOPED_TIMING_STAT(Networking, FastArray_Apply);
+
 	// ---------------------------------------------------------
 	// Look for implicit deletes that would happen due to Naks
 	// ---------------------------------------------------------
diff --git a/Engine/Source/Runtime/Net/Core/Private/Net/Core/NetCoreModule.cpp b/Engine/Source/Runtime/Net/Core/Private/Net/Core/NetCoreModule.cpp
--- a/Engine/Source/Runtime/Net/Core/Private/Net/Core/NetCoreModule.cpp
+++ b/Engine/Source/Runtime/Net/Core/Private/Net/Core/NetCoreModule.cpp
@@ -9,6 +9,7 @@
 #include "Net/Core/DirtyNetObjectTracker/GlobalDirtyNetObjectTracker.h"
 #include "Net/Core/NetHandle/NetHandleManager.h"
 #include "Net/Serialization/FastArraySerializer.h"
+#include "ProfilingDebugging/CsvProfiler.h"
 
 bool GUseDetailedScopeCounters=true;
 static FAutoConsoleVariableRef CVarNetEnableDetailedScopeCounters(
@@ -21,6 +22,8 @@ DEFINE_LOG_CATEGORY(LogNetFastTArray);
 
 DEFINE_LOG_CATEGORY(LogNetCore);
 
+CSV_DEFINE_CATEGORY_MODULE(NETCORE_API, Networking, true);
+
 
 class FNetCoreModule : public IModuleInterface
 {
