Optimized PreReplication to not lookup fRepChangedPropertyTracker unless required.
This also allows SetPropertyActiveOverride macros to be called outside of PreReplicate to move towards proper push model.
Removed some old deprecated code

diff --git a/Engine/Source/Runtime/CoreUObject/Public/UObject/CoreNet.h b/Engine/Source/Runtime/CoreUObject/Public/UObject/CoreNet.h
--- a/Engine/Source/Runtime/CoreUObject/Public/UObject/CoreNet.h
+++ b/Engine/Source/Runtime/CoreUObject/Public/UObject/CoreNet.h
@@ -540,24 +540,9 @@ public:
 class IRepChangedPropertyTracker
 {
 public:
+	UE_DEPRECATED(5.5, "Will be removed or renamed and used to pass paramters.")
 	IRepChangedPropertyTracker() { }
 	virtual ~IRepChangedPropertyTracker() { }
-
-	UE_DEPRECATED(5.3, "Please use FPropertyConditions::SetActiveOverride instead.")
-	virtual void SetCustomIsActiveOverride(UObject* OwningObject, const uint16 RepIndex, const bool bIsActive) = 0;
-
-	/**
-	* Used when tracking memory to gather the total size of a given instance.
-	* This should include the dynamically allocated data, as well as the classes size.
-	*/
-	virtual void CountBytes(FArchive& Ar) const {};
-
-private:
-	PRAGMA_DISABLE_DEPRECATION_WARNINGS
-	void CallSetCustomIsActiveOverride(UObject* OwningObject, const uint16 RepIndex, const bool bIsActive) { SetCustomIsActiveOverride(OwningObject, RepIndex, bIsActive); }
-	PRAGMA_ENABLE_DEPRECATION_WARNINGS
-	
-	friend UE::Net::Private::FNetPropertyConditionManager;
 };
 
 class FCustomPropertyConditionState
diff --git a/Engine/Source/Runtime/Engine/Classes/Engine/NetDriver.h b/Engine/Source/Runtime/Engine/Classes/Engine/NetDriver.h
--- a/Engine/Source/Runtime/Engine/Classes/Engine/NetDriver.h
+++ b/Engine/Source/Runtime/Engine/Classes/Engine/NetDriver.h
@@ -1259,36 +1259,6 @@ public:
 	 */
 	TMap<FName, FName>	RenamedStartupActors;
 
-	class UE_DEPRECATED(5.1, "No longer used.") FRepChangedPropertyTrackerWrapper
-	{
-	public:
-		FRepChangedPropertyTrackerWrapper(UObject* Obj, const TSharedPtr<FRepChangedPropertyTracker>& InRepChangedPropertyTracker) : RepChangedPropertyTracker(InRepChangedPropertyTracker), WeakObjectPtr(Obj) {}
-
-		const FRepChangedPropertyTracker* operator->() const { return RepChangedPropertyTracker.Get(); }
-		FRepChangedPropertyTracker* operator->() { return RepChangedPropertyTracker.Get(); }
-		
-		const FRepChangedPropertyTracker* Get() const { return RepChangedPropertyTracker.Get(); }
-		FRepChangedPropertyTracker* Get() { return RepChangedPropertyTracker.Get(); }
-
-		bool IsValid() const { return RepChangedPropertyTracker.IsValid(); }
-		bool IsObjectValid() const { return WeakObjectPtr.IsValid(); }
-
-		TWeakObjectPtr<UObject> GetWeakObjectPtr() const { return WeakObjectPtr; }
-
-		TSharedPtr<FRepChangedPropertyTracker> RepChangedPropertyTracker;
-
-		void CountBytes(FArchive& Ar) const;
-
-	private:
-		TWeakObjectPtr<UObject> WeakObjectPtr;
-	};
-
-	PRAGMA_DISABLE_DEPRECATION_WARNINGS
-	/** Maps FRepChangedPropertyTracker to active objects that are replicating properties */
-	UE_DEPRECATED(5.1, "Property trackers have been moved to the NetCore module")
-	TMap<UObject*, FRepChangedPropertyTrackerWrapper>	RepChangedPropertyTrackerMap;
-	PRAGMA_ENABLE_DEPRECATION_WARNINGS
-
 	/** Used to invalidate properties marked "unchanged" in FRepChangedPropertyTracker's */
 	uint32 ReplicationFrame;
 
diff --git a/Engine/Source/Runtime/Engine/Private/Actor.cpp b/Engine/Source/Runtime/Engine/Private/Actor.cpp
--- a/Engine/Source/Runtime/Engine/Private/Actor.cpp
+++ b/Engine/Source/Runtime/Engine/Private/Actor.cpp
@@ -1582,7 +1582,13 @@ void AActor::CallPreReplication(UNetDriver* NetDriver)
 	const bool bPreReplication = ShouldCallPreReplication();
 	const bool bPreReplicationForReplay = ShouldCallPreReplicationForReplay();
 
-	TSharedPtr<FRepChangedPropertyTracker> ActorChangedPropertyTracker;
+PRAGMA_DISABLE_DEPRECATION_WARNINGS
+	// This interface is obsolete and will be removed or re-purposed and renamed to pass parameters to PreReplication.
+	IRepChangedPropertyTracker DummyRepChangedPropertyTracker;
+PRAGMA_ENABLE_DEPRECATION_WARNINGS
+
+	// Allow lazy creation/find of FRepChangedPropertyTracker
+	UE::Net::Private::FNetPropertyConditionManager::FAllowCreateTrackerFromSetPropertyActiveOverrideScope AllowCreateTracker(UE::Net::Private::FNetPropertyConditionManager::Get());	
 
 	if (bPreReplication)
 	{
@@ -1593,8 +1599,7 @@ void AActor::CallPreReplication(UNetDriver* NetDriver)
 		// In that case we call PreReplication on the locally controlled Character as well.
 		if ((LocalRole == ROLE_Authority) || ((LocalRole == ROLE_AutonomousProxy) && World && World->IsRecordingClientReplay()))
 		{
-			ActorChangedPropertyTracker = NetDriver->FindOrCreateRepChangedPropertyTracker(this);
-			PreReplication(*(ActorChangedPropertyTracker.Get()));
+			PreReplication(DummyRepChangedPropertyTracker);
 		}
 	}
 
@@ -1603,12 +1608,7 @@ void AActor::CallPreReplication(UNetDriver* NetDriver)
 		// If we're recording a replay, call this for everyone (includes SimulatedProxies).
 		if (Cast<UDemoNetDriver>(NetDriver) || NetDriver->HasReplayConnection())
 		{
-			if (!ActorChangedPropertyTracker.IsValid())
-			{
-				ActorChangedPropertyTracker = NetDriver->FindOrCreateRepChangedPropertyTracker(this);
-			}
-
-			PreReplicationForReplay(*(ActorChangedPropertyTracker.Get()));
+			PreReplicationForReplay(DummyRepChangedPropertyTracker);
 		}
 	}
 
@@ -1620,8 +1620,7 @@ void AActor::CallPreReplication(UNetDriver* NetDriver)
 			// Only call on components that aren't pending kill
 			if (IsValid(Component))
 			{
-				TSharedPtr<FRepChangedPropertyTracker> ComponentChangedPropertyTracker = NetDriver->FindOrCreateRepChangedPropertyTracker(Component);
-				Component->PreReplication(*(ComponentChangedPropertyTracker.Get()));
+				Component->PreReplication(DummyRepChangedPropertyTracker);
 			}
 		}
 	}
diff --git a/Engine/Source/Runtime/Engine/Private/NetDriver.cpp b/Engine/Source/Runtime/Engine/Private/NetDriver.cpp
--- a/Engine/Source/Runtime/Engine/Private/NetDriver.cpp
+++ b/Engine/Source/Runtime/Engine/Private/NetDriver.cpp
@@ -2988,14 +2988,6 @@ void UNetDriver::SetAnalyticsProvider(TSharedPtr<IAnalyticsProvider> InProvider)
 	}
 }
 
-void UNetDriver::FRepChangedPropertyTrackerWrapper::CountBytes(FArchive& Ar) const
-{
-	if (FRepChangedPropertyTracker const * const LocalTracker = RepChangedPropertyTracker.Get())
-	{
-		LocalTracker->CountBytes(Ar);
-	}
-}
-
 void UNetDriver::FReplicationChangelistMgrWrapper::CountBytes(FArchive& Ar) const
 {
 	if (FReplicationChangelistMgr const * const ChangelistMgr = ReplicationChangelistMgr.Get())
@@ -3070,17 +3062,6 @@ void UNetDriver::Serialize( FArchive& Ar )
 
 		GRANULAR_NETWORK_MEMORY_TRACKING_TRACK("RenamedStartupActors", RenamedStartupActors.CountBytes(Ar));
 
-		PRAGMA_DISABLE_DEPRECATION_WARNINGS
-		GRANULAR_NETWORK_MEMORY_TRACKING_TRACK("RepChangedPropertyTrackerMap",
-			RepChangedPropertyTrackerMap.CountBytes(Ar);
-
-			for (const auto& RepChangedPropertyTrackerPair : RepChangedPropertyTrackerMap)
-			{
-				RepChangedPropertyTrackerPair.Value.CountBytes(Ar);
-			}
-		);
-		PRAGMA_ENABLE_DEPRECATION_WARNINGS
-
 		GRANULAR_NETWORK_MEMORY_TRACKING_TRACK("RepLayoutMap",
 			RepLayoutMap.CountBytes(Ar);
 
@@ -4353,16 +4334,6 @@ void UNetDriver::PostGarbageCollect()
 		}
 	}
 
-	PRAGMA_DISABLE_DEPRECATION_WARNINGS
-	for (auto It = RepChangedPropertyTrackerMap.CreateIterator(); It; ++It)
-	{
-		if (!It.Value().IsObjectValid())
-		{
-			It.RemoveCurrent();
-		}
-	}
-	PRAGMA_ENABLE_DEPRECATION_WARNINGS
-
 	for (FObjectReplicator* Replicator : AllOwnedReplicators)
 	{
 		if (!Replicator->GetWeakObjectPtr().IsValid())
diff --git a/Engine/Source/Runtime/Engine/Public/Net/UnrealNetwork.h b/Engine/Source/Runtime/Engine/Public/Net/UnrealNetwork.h
--- a/Engine/Source/Runtime/Engine/Public/Net/UnrealNetwork.h
+++ b/Engine/Source/Runtime/Engine/Public/Net/UnrealNetwork.h
@@ -297,15 +297,16 @@ inline FProperty* GetReplicatedProperty(const UClass* CallingClass, const UClass
 #define DOREPLIFETIME_ACTIVE_OVERRIDE_FAST(c,v,active) \
 { \
 	static_assert(ValidateReplicatedClassInheritance<c, ThisClass>(), #c "." #v " is not accessible from this class."); \
-	UE::Net::Private::FNetPropertyConditionManager::SetPropertyActiveOverride(ChangedPropertyTracker, this, (int32)c::ENetFields_Private::v, active); \
+	UE::Net::Private::FNetPropertyConditionManager::Get().SetPropertyActiveOverride(this, (int32)c::ENetFields_Private::v, active); \
 }
 
 #define DOREPLIFETIME_ACTIVE_OVERRIDE_FAST_STATIC_ARRAY(c,v,active) \
 { \
 	static_assert(ValidateReplicatedClassInheritance<c, ThisClass>(), #c "." #v " is not accessible from this class."); \
+	UE::Net::Private::FNetPropertyConditionManager& PropertyConditionManager = UE::Net::Private::FNetPropertyConditionManager::Get(); \
 	for (int32 i = 0; i < (int32)c::EArrayDims_Private::v; ++i) \
 	{ \
-		UE::Net::Private::FNetPropertyConditionManager::SetPropertyActiveOverride(ChangedPropertyTracker, this, (int32)c::ENetFields_Private::v##_STATIC_ARRAY + i, active); \
+		PropertyConditionManager.SetPropertyActiveOverride(this, (int32)c::ENetFields_Private::v##_STATIC_ARRAY + i, active); \
 	} \
 }
 
@@ -313,9 +314,10 @@ inline FProperty* GetReplicatedProperty(const UClass* CallingClass, const UClass
 { \
 	static_assert(ValidateReplicatedClassInheritance<c, ThisClass>(), #c "." #v " is not accessible from this class."); \
 	static FProperty* sp##v = GetReplicatedProperty(StaticClass(), c::StaticClass(),GET_MEMBER_NAME_CHECKED(c,v)); \
+	UE::Net::Private::FNetPropertyConditionManager& PropertyConditionManager = UE::Net::Private::FNetPropertyConditionManager::Get(); \
 	for (int32 i = 0; i < sp##v->ArrayDim; i++) \
 	{ \
-		UE::Net::Private::FNetPropertyConditionManager::SetPropertyActiveOverride(ChangedPropertyTracker, this, sp##v->RepIndex + i, active); \
+		PropertyConditionManager.SetPropertyActiveOverride(this, sp##v->RepIndex + i, active); \
 	} \
 }
 
diff --git a/Engine/Source/Runtime/Net/Core/Private/Net/Core/PropertyConditions/PropertyConditions.cpp b/Engine/Source/Runtime/Net/Core/Private/Net/Core/PropertyConditions/PropertyConditions.cpp
--- a/Engine/Source/Runtime/Net/Core/Private/Net/Core/PropertyConditions/PropertyConditions.cpp
+++ b/Engine/Source/Runtime/Net/Core/Private/Net/Core/PropertyConditions/PropertyConditions.cpp
@@ -33,26 +33,54 @@ FNetPropertyConditionManager& FNetPropertyConditionManager::Get()
 	return Singleton;
 }
 
-void FNetPropertyConditionManager::SetPropertyActive(const FObjectKey ObjectKey, const uint16 RepIndex, const bool bActive)
+void FNetPropertyConditionManager::SetPropertyActive(const UObject* Object, const uint16 RepIndex, const bool bActive)
 {
+	const FObjectKey ObjectKey(Object);
+
 	TSharedPtr<FRepChangedPropertyTracker> Tracker = FindPropertyTracker(ObjectKey);
 	if (Tracker.IsValid())
 	{
-		Tracker->CallSetCustomIsActiveOverride(ObjectKey.ResolveObjectPtr(), RepIndex, bActive);
+		Tracker->SetCustomIsActiveOverride(Object, RepIndex, bActive);
 	}
 }
 
-void FNetPropertyConditionManager::SetPropertyDynamicCondition(const FObjectKey ObjectKey, const uint16 RepIndex, const ELifetimeCondition Condition)
+void FNetPropertyConditionManager::SetPropertyActiveOverride(const UObject* Object, const uint16 RepIndex, const bool bActive)
 {
+	const FObjectKey ObjectKey(Object);
+
+	// We cache the LastFoundTracker and Key to avoid the map lookup if we modify multiple properties at once.
+	if (LastFoundTrackerKey != ObjectKey || LastFoundTracker == nullptr)
+	{
+		TSharedPtr<FRepChangedPropertyTracker> Tracker = bAllowCreateTrackerFromSetPropertyActiveOverride ? FindOrCreatePropertyTracker(ObjectKey) : FindPropertyTracker(ObjectKey);
+		LastFoundTrackerKey = ObjectKey;
+		LastFoundTracker = Tracker.Get();
+	}
+
+	if (LastFoundTracker)
+	{
+		LastFoundTracker->SetCustomIsActiveOverride(Object, RepIndex, bActive);
+	}
+}
+
+void FNetPropertyConditionManager::SetPropertyDynamicCondition(const UObject* Object, const uint16 RepIndex, const ELifetimeCondition Condition)
+{
+	const FObjectKey ObjectKey(Object);
+
 	TSharedPtr<FRepChangedPropertyTracker> Tracker = FindPropertyTracker(ObjectKey);
 	if (Tracker.IsValid())
 	{
-		Tracker->CallSetDynamicCondition(ObjectKey.ResolveObjectPtr(), RepIndex, Condition);
+		Tracker->SetDynamicCondition(Object, RepIndex, Condition);
 	}
 }
 
 void FNetPropertyConditionManager::NotifyObjectDestroyed(const FObjectKey ObjectKey)
 {
+	if (ObjectKey == LastFoundTrackerKey)
+	{
+		LastFoundTrackerKey;
+		LastFoundTracker = nullptr;
+	}
+
 	PropertyTrackerMap.Remove(ObjectKey);
 }
 
@@ -98,6 +126,11 @@ void FNetPropertyConditionManager::PostGarbageCollect()
 	{
 		if (!It.Key().ResolveObjectPtr())
 		{
+			if (It.Key() == LastFoundTrackerKey)
+			{
+				LastFoundTrackerKey;
+				LastFoundTracker = nullptr;
+			}
 			It.RemoveCurrent();
 		}
 	}
@@ -122,9 +155,4 @@ void FNetPropertyConditionManager::LogMemory(FOutputDevice& Ar)
 	Ar.Logf(TEXT("  Property Condition Memory: %u"), CountBytes);
 }
 
-void FNetPropertyConditionManager::SetPropertyActiveOverride(IRepChangedPropertyTracker& Tracker, UObject* OwningObject, const uint16 RepIndex, const bool bIsActive)
-{
-	Tracker.CallSetCustomIsActiveOverride(OwningObject, RepIndex, bIsActive);
-}
-
 }; // UE::Net::Private
\ No newline at end of file
diff --git a/Engine/Source/Runtime/Net/Core/Private/Net/Core/PropertyConditions/RepChangedPropertyTracker.cpp b/Engine/Source/Runtime/Net/Core/Private/Net/Core/PropertyConditions/RepChangedPropertyTracker.cpp
--- a/Engine/Source/Runtime/Net/Core/Private/Net/Core/PropertyConditions/RepChangedPropertyTracker.cpp
+++ b/Engine/Source/Runtime/Net/Core/Private/Net/Core/PropertyConditions/RepChangedPropertyTracker.cpp
@@ -9,11 +9,7 @@ FRepChangedPropertyTracker::FRepChangedPropertyTracker(FCustomPropertyConditionS
 {
 }
 
-FRepChangedPropertyTracker::FRepChangedPropertyTracker(const bool InbIsReplay, const bool InbIsClientReplayRecording)
-	: FRepChangedPropertyTracker(0)
-{}
-
-void FRepChangedPropertyTracker::SetCustomIsActiveOverride(UObject* OwningObject, const uint16 RepIndex, const bool bIsActive)
+void FRepChangedPropertyTracker::SetCustomIsActiveOverride(const UObject* OwningObject, const uint16 RepIndex, const bool bIsActive)
 {
 	const bool bOldActive = ActiveState.GetActiveState(RepIndex);
 	ActiveState.SetActiveState(RepIndex, bIsActive);	// check for client replay recording moved to FReplicationFlags
@@ -33,7 +29,7 @@ void FRepChangedPropertyTracker::SetCustomIsActiveOverride(UObject* OwningObject
 #endif // UE_WITH_IRIS
 }
 
-void FRepChangedPropertyTracker::CallSetDynamicCondition(const UObject* OwningObject, const uint16 RepIndex, const ELifetimeCondition Condition)
+void FRepChangedPropertyTracker::SetDynamicCondition(const UObject* OwningObject, const uint16 RepIndex, const ELifetimeCondition Condition)
 {
 	using namespace UE::Net::Private;
 
diff --git a/Engine/Source/Runtime/Net/Core/Public/Net/Core/PropertyConditions/PropertyConditions.h b/Engine/Source/Runtime/Net/Core/Public/Net/Core/PropertyConditions/PropertyConditions.h
--- a/Engine/Source/Runtime/Net/Core/Public/Net/Core/PropertyConditions/PropertyConditions.h
+++ b/Engine/Source/Runtime/Net/Core/Public/Net/Core/PropertyConditions/PropertyConditions.h
@@ -24,9 +24,10 @@ public:
 
 	static NETCORE_API FNetPropertyConditionManager& Get();
 
-	NETCORE_API void SetPropertyActive(const FObjectKey ObjectKey, const uint16 RepIndex, const bool bActive);
-	NETCORE_API void SetPropertyDynamicCondition(const FObjectKey ObjectKey, const uint16 RepIndex, const ELifetimeCondition Condition);
-
+	NETCORE_API void SetPropertyActive(const UObject* Object, const uint16 RepIndex, const bool bActive);
+	NETCORE_API void SetPropertyActiveOverride(const UObject* Object, const uint16 RepIndex, const bool bIsActive);
+	NETCORE_API void SetPropertyDynamicCondition(const UObject* Object, const uint16 RepIndex, const ELifetimeCondition Condition);
+	
 	NETCORE_API void NotifyObjectDestroyed(const FObjectKey ObjectKey);
 
 	NETCORE_API TSharedPtr<FRepChangedPropertyTracker> FindOrCreatePropertyTracker(const FObjectKey ObjectKey);
@@ -34,14 +35,27 @@ public:
 
 	NETCORE_API void LogMemory(FOutputDevice& Ar);
 
-	static NETCORE_API void SetPropertyActiveOverride(IRepChangedPropertyTracker& Tracker, UObject* OwningObject, const uint16 RepIndex, const bool bIsActive);
+	class FAllowCreateTrackerFromSetPropertyActiveOverrideScope
+	{
+	public:
+		FAllowCreateTrackerFromSetPropertyActiveOverrideScope(FNetPropertyConditionManager& ConditionManager)
+		: GuardValue(ConditionManager.bAllowCreateTrackerFromSetPropertyActiveOverride, true) {};
+	private:
+		TGuardValue<bool> GuardValue;
+	};
 
 private:
+	friend class FAllowCreateTrackerFromSetPropertyActiveOverrideScope;
+
 	void PostGarbageCollect();
 
 	FDelegateHandle PostGarbageCollectHandle;
 
 	TMap<FObjectKey, TSharedPtr<FRepChangedPropertyTracker>> PropertyTrackerMap;
+
+	FObjectKey LastFoundTrackerKey;
+	FRepChangedPropertyTracker* LastFoundTracker = nullptr;
+	bool bAllowCreateTrackerFromSetPropertyActiveOverride = false;
 };
 
 }; // UE::Net::Private
diff --git a/Engine/Source/Runtime/Net/Core/Public/Net/Core/PropertyConditions/RepChangedPropertyTracker.h b/Engine/Source/Runtime/Net/Core/Public/Net/Core/PropertyConditions/RepChangedPropertyTracker.h
--- a/Engine/Source/Runtime/Net/Core/Public/Net/Core/PropertyConditions/RepChangedPropertyTracker.h
+++ b/Engine/Source/Runtime/Net/Core/Public/Net/Core/PropertyConditions/RepChangedPropertyTracker.h
@@ -19,36 +19,15 @@ class UObject;
  * TODO: This class (and arguably IRepChangedPropertyTracker) should be renamed to reflect
  *			what they actually do now.
  */
-class FRepChangedPropertyTracker : public IRepChangedPropertyTracker
+class FRepChangedPropertyTracker
 {
 public:
 	FRepChangedPropertyTracker() = delete;
 	FRepChangedPropertyTracker(FCustomPropertyConditionState&& InActiveState);
 
-	UE_DEPRECATED(5.1, "Replay arguments no longer used.")
-	FRepChangedPropertyTracker(const bool InbIsReplay, const bool InbIsClientReplayRecording);
-
-	virtual ~FRepChangedPropertyTracker() = default;
-
-	//~ Begin IRepChangedPropertyTracker Interface.
-	/**
-		* Manually set whether or not Property should be marked inactive.
-		* This will change the Active status for all connections.
-		*
-		* @see DOREPLIFETIME_ACTIVE_OVERRIDE
-		*
-		* @param OwningObject	The object that we're tracking.
-		* @param RepIndex		Replication index for the Property.
-		* @param bIsActive		The new Active state.
-		*/
-	UE_DEPRECATED(5.2, "Please use FNetPropertyConditionManager::SetPropertyActiveOverride instead.")
-	virtual void SetCustomIsActiveOverride(UObject* OwningObject, const uint16 RepIndex, const bool bIsActive) override;
-
-	virtual void CountBytes(FArchive& Ar) const override;
-	//~ End IRepChangedPropertyTracker Interface
-
-	UE_DEPRECATED(5.1, "No longer used, ActiveState must be constructed with the correct number of properties.")
-	void InitActiveParents(int32 ParentCount) {}
+	~FRepChangedPropertyTracker() = default;
+	
+	void CountBytes(FArchive& Ar) const;
 
 	bool IsParentActive(uint16 ParentIndex) const
 	{
@@ -74,7 +53,8 @@ private:
 	friend UE::Net::Private::FNetPropertyConditionManager;
 
 	// Called from FNetPropertyConditionManager 
-	void CallSetDynamicCondition(const UObject* OwningObject, const uint16 RepIndex, const ELifetimeCondition Condition);
+	void SetDynamicCondition(const UObject* OwningObject, const uint16 RepIndex, const ELifetimeCondition Condition);
+	void SetCustomIsActiveOverride(const UObject* OwningObject, const uint16 RepIndex, const bool bIsActive);
 
 	/** Activation data for top level Properties on the given Actor / Object. */
 	FCustomPropertyConditionState ActiveState;
