Iris
- Fixed error in PropertyConditions.cpp

diff --git a/Engine/Source/Runtime/Net/Core/Private/Net/Core/PropertyConditions/PropertyConditions.cpp b/Engine/Source/Runtime/Net/Core/Private/Net/Core/PropertyConditions/PropertyConditions.cpp
--- a/Engine/Source/Runtime/Net/Core/Private/Net/Core/PropertyConditions/PropertyConditions.cpp
+++ b/Engine/Source/Runtime/Net/Core/Private/Net/Core/PropertyConditions/PropertyConditions.cpp
@@ -77,7 +77,7 @@ void FNetPropertyConditionManager::NotifyObjectDestroyed(const FObjectKey Object
 {
 	if (ObjectKey == LastFoundTrackerKey)
 	{
-		LastFoundTrackerKey;
+		LastFoundTrackerKey = FObjectKey();
 		LastFoundTracker = nullptr;
 	}
 
@@ -128,7 +128,7 @@ void FNetPropertyConditionManager::PostGarbageCollect()
 		{
 			if (It.Key() == LastFoundTrackerKey)
 			{
-				LastFoundTrackerKey;
+				LastFoundTrackerKey = FObjectKey();
 				LastFoundTracker = nullptr;
 			}
 			It.RemoveCurrent();
