Fix partial net traces:
- make some net trace events important so they're cached and written to all traces

diff --git a/Engine/Source/Runtime/Net/Core/Private/Net/Core/Trace/NetTraceInternal.cpp b/Engine/Source/Runtime/Net/Core/Private/Net/Core/Trace/NetTraceInternal.cpp
--- a/Engine/Source/Runtime/Net/Core/Private/Net/Core/Trace/NetTraceInternal.cpp
+++ b/Engine/Source/Runtime/Net/Core/Private/Net/Core/Trace/NetTraceInternal.cpp
@@ -112,7 +112,8 @@ void FNetTrace::TraceInstanceUpdated(uint32 GameInstanceId, bool bIsServer, cons
 
 FNetTraceCollector* FNetTrace::CreateTraceCollector()
 {
-	if (!GNetTraceRuntimeVerbosity)
+	if (!GNetTraceRuntimeVerbosity || 
+		!UE::Trace::IsTracing()) // don't use collectors when not actively tracing, for perf reasons
 	{
 		return nullptr;
 	}
@@ -531,7 +532,7 @@ UE::Net::FNetDebugNameId FNetTrace::TraceName(const TCHAR* Name)
 		ThreadBufferPtr->DynamicNameHashToNameIdMap.Add(HashedName, NameId);
 
 		FTCHARToUTF8 Converter(Name);
-		FNetTraceInternal::Reporter::ReportAnsiName(NameId, Converter.Length() + 1, (const char*)Converter.Get());		
+		FNetTraceInternal::Reporter::ReportAnsiName(NameId, Converter.Length(), (const char*)Converter.Get());		
 		
 		return NameId;
 	}
@@ -566,7 +567,7 @@ UE::Net::FNetDebugNameId FNetTrace::TraceName(FName Name)
 		TCHAR Buffer[StringBufferSize];
 		uint32 NameLen = Name.ToString(Buffer);
 		FTCHARToUTF8 Converter(Buffer);
-		FNetTraceInternal::Reporter::ReportAnsiName(NameId, Converter.Length() + 1, (const char*)Converter.Get());		
+		FNetTraceInternal::Reporter::ReportAnsiName(NameId, Converter.Length(), (const char*)Converter.Get());		
 		
 		return NameId;
 	}
diff --git a/Engine/Source/Runtime/Net/Core/Private/Net/Core/Trace/Reporters/NetTraceReporter.cpp b/Engine/Source/Runtime/Net/Core/Private/Net/Core/Trace/Reporters/NetTraceReporter.cpp
--- a/Engine/Source/Runtime/Net/Core/Private/Net/Core/Trace/Reporters/NetTraceReporter.cpp
+++ b/Engine/Source/Runtime/Net/Core/Private/Net/Core/Trace/Reporters/NetTraceReporter.cpp
@@ -16,19 +16,19 @@ uint32 FNetTraceReporter::NetTraceReporterVersion = 1;
 UE_TRACE_CHANNEL_DEFINE(NetChannel)
 
 // We always output this event first to make sure we have a version number for backwards compatibility
-UE_TRACE_EVENT_BEGIN(NetTrace, InitEvent)
+UE_TRACE_EVENT_BEGIN(NetTrace, InitEvent, NoSync|Important)
 	UE_TRACE_EVENT_FIELD(uint64, Timestamp)
 	UE_TRACE_EVENT_FIELD(uint32, NetTraceVersion)
 	UE_TRACE_EVENT_FIELD(uint32, NetTraceReporterVersion)
 UE_TRACE_EVENT_END()
 
 // Trace a name
-UE_TRACE_EVENT_BEGIN(NetTrace, NameEvent)
+UE_TRACE_EVENT_BEGIN(NetTrace, NameEvent, NoSync|Important)
 	UE_TRACE_EVENT_FIELD(uint16, NameId)
 	UE_TRACE_EVENT_FIELD(UE::Trace::AnsiString, Name)
 UE_TRACE_EVENT_END()
 
-UE_TRACE_EVENT_BEGIN(NetTrace, ObjectCreatedEvent)	
+UE_TRACE_EVENT_BEGIN(NetTrace, ObjectCreatedEvent, NoSync|Important)	
 	UE_TRACE_EVENT_FIELD(uint64, TypeId)
 	UE_TRACE_EVENT_FIELD(uint64, ObjectId)
 	UE_TRACE_EVENT_FIELD(uint32, OwnerId)
@@ -36,25 +36,25 @@ UE_TRACE_EVENT_BEGIN(NetTrace, ObjectCreatedEvent)
 	UE_TRACE_EVENT_FIELD(uint8, GameInstanceId)
 UE_TRACE_EVENT_END()
 
-UE_TRACE_EVENT_BEGIN(NetTrace, ObjectDestroyedEvent)
+UE_TRACE_EVENT_BEGIN(NetTrace, ObjectDestroyedEvent, NoSync|Important)
 	UE_TRACE_EVENT_FIELD(uint64, ObjectId)
 	UE_TRACE_EVENT_FIELD(uint8, GameInstanceId)
 UE_TRACE_EVENT_END()
 
 // What else do we want to know? should we maybe call this a connectionEvent instead?
-UE_TRACE_EVENT_BEGIN(NetTrace, ConnectionCreatedEvent)
+UE_TRACE_EVENT_BEGIN(NetTrace, ConnectionCreatedEvent, NoSync|Important)
 	UE_TRACE_EVENT_FIELD(uint16, ConnectionId)
 	UE_TRACE_EVENT_FIELD(uint8, GameInstanceId)
 UE_TRACE_EVENT_END()
 
-UE_TRACE_EVENT_BEGIN(NetTrace, ConnectionStateUpdatedEvent)
+UE_TRACE_EVENT_BEGIN(NetTrace, ConnectionStateUpdatedEvent, NoSync|Important)
 	UE_TRACE_EVENT_FIELD(uint16, ConnectionId)
 	UE_TRACE_EVENT_FIELD(uint8, GameInstanceId)
 	UE_TRACE_EVENT_FIELD(uint8, ConnectionStateValue)
 UE_TRACE_EVENT_END()
 
 // Provides additional information about connection after it is created
-UE_TRACE_EVENT_BEGIN(NetTrace, ConnectionUpdatedEvent)
+UE_TRACE_EVENT_BEGIN(NetTrace, ConnectionUpdatedEvent, NoSync|Important)
 	UE_TRACE_EVENT_FIELD(uint16, ConnectionId)
 	UE_TRACE_EVENT_FIELD(uint8, GameInstanceId)
 	UE_TRACE_EVENT_FIELD(UE::Trace::WideString, Name)
@@ -62,7 +62,7 @@ UE_TRACE_EVENT_BEGIN(NetTrace, ConnectionUpdatedEvent)
 UE_TRACE_EVENT_END()
 
 // Add close reason?
-UE_TRACE_EVENT_BEGIN(NetTrace, ConnectionClosedEvent)
+UE_TRACE_EVENT_BEGIN(NetTrace, ConnectionClosedEvent, NoSync|Important)
 	UE_TRACE_EVENT_FIELD(uint16, ConnectionId)
 	UE_TRACE_EVENT_FIELD(uint8, GameInstanceId)
 UE_TRACE_EVENT_END()
@@ -83,14 +83,14 @@ UE_TRACE_EVENT_BEGIN(NetTrace, FrameStatsCounterEvent)
 UE_TRACE_EVENT_END()
 
 // Provides additional information about game instance
-UE_TRACE_EVENT_BEGIN(NetTrace, InstanceUpdatedEvent)
+UE_TRACE_EVENT_BEGIN(NetTrace, InstanceUpdatedEvent, NoSync|Important)
 	UE_TRACE_EVENT_FIELD(uint8, GameInstanceId)
 	UE_TRACE_EVENT_FIELD(bool, bIsServer)
 	UE_TRACE_EVENT_FIELD(UE::Trace::WideString, Name)
 UE_TRACE_EVENT_END()
 
 // rename 
-UE_TRACE_EVENT_BEGIN(NetTrace, InstanceDestroyedEvent)
+UE_TRACE_EVENT_BEGIN(NetTrace, InstanceDestroyedEvent, NoSync|Important)
 	UE_TRACE_EVENT_FIELD(uint8, GameInstanceId)
 UE_TRACE_EVENT_END()
 
@@ -130,10 +130,11 @@ void FNetTraceReporter::ReportInitEvent(uint32 NetTraceVersion)
 
 void FNetTraceReporter::ReportInstanceUpdated(uint32 GameInstanceId, bool bIsServer, const TCHAR* Name)
 {
-	UE_TRACE_LOG(NetTrace, InstanceUpdatedEvent, NetChannel)
+	int32 NameLen = FCString::Strlen(Name);
+	UE_TRACE_LOG(NetTrace, InstanceUpdatedEvent, NetChannel, NameLen * sizeof(TCHAR))
 		<< InstanceUpdatedEvent.GameInstanceId((uint8)GameInstanceId)
 		<< InstanceUpdatedEvent.bIsServer(bIsServer)
-		<< InstanceUpdatedEvent.Name(Name);
+		<< InstanceUpdatedEvent.Name(Name, NameLen);
 }
 
 void FNetTraceReporter::ReportInstanceDestroyed(uint32 GameInstanceId)
@@ -142,11 +143,11 @@ void FNetTraceReporter::ReportInstanceDestroyed(uint32 GameInstanceId)
 		<< InstanceDestroyedEvent.GameInstanceId((uint8)GameInstanceId);
 }
 
-void FNetTraceReporter::ReportAnsiName(UE::Net::FNetDebugNameId NameId, uint32 NameSize, const char* Name)
+void FNetTraceReporter::ReportAnsiName(UE::Net::FNetDebugNameId NameId, uint32 NameLen, const char* Name)
 {
-	UE_TRACE_LOG(NetTrace, NameEvent, NetChannel)
+	UE_TRACE_LOG(NetTrace, NameEvent, NetChannel, NameLen)
 		<< NameEvent.NameId(NameId)
-		<< NameEvent.Name(Name, NameSize);
+		<< NameEvent.Name(Name, NameLen);
 }
 
 void FNetTraceReporter::ReportPacketDropped(const FNetTracePacketInfo& PacketInfo)
@@ -302,11 +303,13 @@ void FNetTraceReporter::ReportConnectionStateUpdated(uint32 GameInstanceId, uint
 
 void FNetTraceReporter::ReportConnectionUpdated(uint32 GameInstanceId, uint32 ConnectionId, const TCHAR* AddressString, const TCHAR* OwningActor)
 {
-	UE_TRACE_LOG(NetTrace, ConnectionUpdatedEvent, NetChannel)
+	int32 AddressLen = FCString::Strlen(AddressString);
+	int32 OwningActorLen = FCString::Strlen(OwningActor);
+	UE_TRACE_LOG(NetTrace, ConnectionUpdatedEvent, NetChannel, (AddressLen + OwningActorLen) * sizeof(TCHAR))
 		<< ConnectionUpdatedEvent.GameInstanceId((uint8)GameInstanceId)
 		<< ConnectionUpdatedEvent.ConnectionId((uint16)ConnectionId)
-		<< ConnectionUpdatedEvent.Name(OwningActor)
-		<< ConnectionUpdatedEvent.Address(AddressString);
+		<< ConnectionUpdatedEvent.Name(OwningActor, OwningActorLen)
+		<< ConnectionUpdatedEvent.Address(AddressString, AddressLen);
 }
 
 void FNetTraceReporter::ReportConnectionClosed(uint32 GameInstanceId, uint32 ConnectionId)
