Export NetChannel from netcore module

diff --git a/Engine/Source/Runtime/Net/Core/Private/Net/Core/Trace/Reporters/NetTraceReporter.cpp b/Engine/Source/Runtime/Net/Core/Private/Net/Core/Trace/Reporters/NetTraceReporter.cpp
--- a/Engine/Source/Runtime/Net/Core/Private/Net/Core/Trace/Reporters/NetTraceReporter.cpp
+++ b/Engine/Source/Runtime/Net/Core/Private/Net/Core/Trace/Reporters/NetTraceReporter.cpp
@@ -13,6 +13,7 @@
 
 uint32 FNetTraceReporter::NetTraceReporterVersion = 1;
 
+NETCORE_API UE_TRACE_CHANNEL_EXTERN(NetChannel)
 UE_TRACE_CHANNEL_DEFINE(NetChannel)
 
 // We always output this event first to make sure we have a version number for backwards compatibility
