[Cleanup] Simplify FOutPacketTraits a bit.

We now assign default values for each field, rather than setting defaults via a dedicated constructor. This is safer and allows copy/move-constructors for FOutPacketTraits to be automatically synthesized by the compiler.

diff --git a/Engine/Source/Runtime/Net/Common/Public/Net/Common/Packets/PacketTraits.h b/Engine/Source/Runtime/Net/Common/Public/Net/Common/Packets/PacketTraits.h
--- a/Engine/Source/Runtime/Net/Common/Public/Net/Common/Packets/PacketTraits.h
+++ b/Engine/Source/Runtime/Net/Common/Public/Net/Common/Packets/PacketTraits.h
@@ -8,12 +8,9 @@
 
 
 // Forward declarations
-namespace UE
+namespace UE::Net
 {
-	namespace Net
-	{
-		struct FNetResult;
-	}
+	struct FNetResult;
 }
 
 
@@ -25,33 +22,22 @@ struct FOutPacketTraits
 	/** Flags - may trigger modification of packet and traits */
 
 	/** Whether or not the packet can/should be compressed */
-	bool bAllowCompression;
+	bool bAllowCompression = true;
 
 
 	/** Traits */
 
 	/** The number of ack bits in the packet - reflecting UNetConnection.NumAckBits */
-	uint32 NumAckBits;
+	uint32 NumAckBits = 0;
 
 	/** The number of bunch bits in the packet - reflecting UNetConnection.NumBunchBits */
-	uint32 NumBunchBits;
+	uint32 NumBunchBits = 0;
 
 	/** Whether or not this is a keepalive packet */
-	bool bIsKeepAlive;
+	bool bIsKeepAlive = false;
 
 	/** Whether or not the packet has been compressed */
-	bool bIsCompressed;
-
-
-	/** Default constructor */
-	FOutPacketTraits()
-		: bAllowCompression(true)
-		, NumAckBits(0)
-		, NumBunchBits(0)
-		, bIsKeepAlive(false)
-		, bIsCompressed(false)
-	{
-	}
+	bool bIsCompressed = false;
 };
 
 /**
