NetCore - Implement FConnectionHandle equal operator manually.

#rnx

diff --git a/Engine/Source/Runtime/Net/Core/Public/Net/Core/Connection/ConnectionHandle.h b/Engine/Source/Runtime/Net/Core/Public/Net/Core/Connection/ConnectionHandle.h
--- a/Engine/Source/Runtime/Net/Core/Public/Net/Core/Connection/ConnectionHandle.h
+++ b/Engine/Source/Runtime/Net/Core/Public/Net/Core/Connection/ConnectionHandle.h
@@ -55,7 +55,10 @@ inline FConnectionHandle::FConnectionHandle(uint32 InParentConnectionId, uint32
 {
 }
 
-inline bool FConnectionHandle::operator==(const FConnectionHandle& Other) const = default;
+inline bool FConnectionHandle::operator==(const FConnectionHandle& Other) const
+{
+	return ParentConnectionId == Other.ParentConnectionId && ChildConnectionId == Other.ChildConnectionId;
+}
 
 inline bool FConnectionHandle::IsValid() const
 {
