Fixing networking GC callback migrating remote objects back

diff --git a/Engine/Source/Runtime/Net/Core/Private/Net/Core/PropertyConditions/PropertyConditions.cpp b/Engine/Source/Runtime/Net/Core/Private/Net/Core/PropertyConditions/PropertyConditions.cpp
--- a/Engine/Source/Runtime/Net/Core/Private/Net/Core/PropertyConditions/PropertyConditions.cpp
+++ b/Engine/Source/Runtime/Net/Core/Private/Net/Core/PropertyConditions/PropertyConditions.cpp
@@ -124,7 +124,11 @@ void FNetPropertyConditionManager::PostGarbageCollect()
 
 	for (auto It = PropertyTrackerMap.CreateIterator(); It; ++It)
 	{
-		if (!It.Key().ResolveObjectPtr())
+		if (
+#if UE_WITH_REMOTE_OBJECT_HANDLE
+			!FObjectPtr(It.Key().GetRemoteId()).IsResolved() ||
+#endif
+			!It.Key().ResolveObjectPtr())
 		{
 			if (It.Key() == LastFoundTrackerKey)
 			{
