Fix potential buffer overflows when using FName::ToString()
- A buffer of at least FName::StringBufferSize elements must be used.
- Added deprecation warning to FName::ToString to catch use of the function with a buffer that is too small. We can make this a static_assert in future engine releases.

diff --git a/Engine/Source/Runtime/Core/Public/UObject/NameTypes.h b/Engine/Source/Runtime/Core/Public/UObject/NameTypes.h
--- a/Engine/Source/Runtime/Core/Public/UObject/NameTypes.h
+++ b/Engine/Source/Runtime/Core/Public/UObject/NameTypes.h
@@ -679,6 +679,7 @@ public:
 	template<int N>
 	uint32 ToString(TCHAR (&Out)[N]) const
 	{
+		UE_STATIC_DEPRECATE(5.6, N < StringBufferSize, "FName::ToString(TCHAR (&Out)[N]) requires a buffer of size of at least FName::StringBufferSize. Passing a smaller buffer can lead to buffer-overflow and crashes.");
 		return ToString(Out, N);
 	}
 
diff --git a/Engine/Source/Runtime/Engine/Private/Animation/AnimTrace.cpp b/Engine/Source/Runtime/Engine/Private/Animation/AnimTrace.cpp
--- a/Engine/Source/Runtime/Engine/Private/Animation/AnimTrace.cpp
+++ b/Engine/Source/Runtime/Engine/Private/Animation/AnimTrace.cpp
@@ -569,7 +569,7 @@ uint32 FAnimTrace::OutputName(const FName& InName)
 
 	if(bShouldTrace)
 	{
-		TCHAR Buffer[256];
+		TCHAR Buffer[FName::StringBufferSize];
 		int32 NameStringLength = InName.ToString(Buffer);
 
 		UE_TRACE_LOG(Animation, Name, AnimationChannel, NameStringLength * sizeof(TCHAR))
diff --git a/Engine/Source/Runtime/Net/Core/Private/Net/Core/Trace/NetTraceInternal.cpp b/Engine/Source/Runtime/Net/Core/Private/Net/Core/Trace/NetTraceInternal.cpp
--- a/Engine/Source/Runtime/Net/Core/Private/Net/Core/Trace/NetTraceInternal.cpp
+++ b/Engine/Source/Runtime/Net/Core/Private/Net/Core/Trace/NetTraceInternal.cpp
@@ -556,8 +556,7 @@ UE::Net::FNetDebugNameId FNetTrace::TraceName(FName Name)
 		const FNetDebugNameId NameId = FNetTraceInternal::GetNextNameId();
 		ThreadBufferPtr->DynamicFNameToNameIdMap.Add(Name, NameId);
 
-		const uint32 StringBufferSize = 256;
-		TCHAR Buffer[StringBufferSize];
+		TCHAR Buffer[FName::StringBufferSize];
 		uint32 NameLen = Name.ToString(Buffer);
 		FTCHARToUTF8 Converter(Buffer);
 		FNetTraceInternal::Reporter::ReportAnsiName(NameId, Converter.Length() + 1, (const char*)Converter.Get());		
diff --git a/Engine/Source/Runtime/RHI/Public/RHIBreadcrumbs.h b/Engine/Source/Runtime/RHI/Public/RHIBreadcrumbs.h
--- a/Engine/Source/Runtime/RHI/Public/RHIBreadcrumbs.h
+++ b/Engine/Source/Runtime/RHI/Public/RHIBreadcrumbs.h
@@ -729,7 +729,7 @@ namespace UE::RHI::GPUProfiler
 
 		struct FConvert
 		{
-			TCHAR Inner[128];
+			TCHAR Inner[FName::StringBufferSize];
 			FConvert(TRHIBreadcrumbValue const& Name)
 			{
 				Name.Value.ToString(Inner);
