Change FPersistentNetDebugNames::CreatePersistentNetDebugName to use a transactionally safe crit section to fix an assert from AutoRTFM

diff --git a/Engine/Source/Runtime/Net/Core/Private/Net/Core/Trace/PersistentNetDebugNames.cpp b/Engine/Source/Runtime/Net/Core/Private/Net/Core/Trace/PersistentNetDebugNames.cpp
--- a/Engine/Source/Runtime/Net/Core/Private/Net/Core/Trace/PersistentNetDebugNames.cpp
+++ b/Engine/Source/Runtime/Net/Core/Private/Net/Core/Trace/PersistentNetDebugNames.cpp
@@ -3,7 +3,7 @@
 #include "Hash/CityHash.h"
 #include "Containers/Map.h"
 #include "Net/Core/Trace/NetTrace.h"
-#include "HAL/CriticalSection.h"
+#include "Misc/TransactionallySafeCriticalSection.h"
 #include "Misc/CoreDelegates.h"
 #include "Misc/MemStack.h"
 #include "Misc/ScopeLock.h"
@@ -24,7 +24,7 @@ public:
 	static void TearDown();
 private:
 
-	FCriticalSection CriticalSection;
+	FTransactionallySafeCriticalSection CriticalSection;
 	TMap<uint64, const FNetDebugName*> PersistentNames;
 	FMemStackBase Allocator;
 	bool bIsInitialized;
@@ -66,7 +66,7 @@ bool FPersistentNetDebugNames::CreatePersistentNetDebugName(const TCHAR* Name, u
 	
 	FNetDebugName* PersistentName = nullptr;
 
-	FScopeLock Lock(&CriticalSection);
+	UE::TScopeLock Lock(CriticalSection);
 	const FNetDebugName** Entry = PersistentNames.Find(HashedName);
 	if (Entry)
 	{
