Make FFastArraySerializer::IDCounter a property so that it's included in object serialization for remote transfers.

diff --git a/Engine/Source/Runtime/Net/Core/Classes/Net/Serialization/FastArraySerializer.h b/Engine/Source/Runtime/Net/Core/Classes/Net/Serialization/FastArraySerializer.h
--- a/Engine/Source/Runtime/Net/Core/Classes/Net/Serialization/FastArraySerializer.h
+++ b/Engine/Source/Runtime/Net/Core/Classes/Net/Serialization/FastArraySerializer.h
@@ -419,6 +419,7 @@ struct FFastArraySerializer
 	TMap<int32, int32> ItemMap;
 
 	/** Counter used to assign IDs to new elements. */
+	UPROPERTY(NotReplicated, Transient)
 	int32 IDCounter;
 
 	/** Counter used to track array replication. */
