Reworking net traces to support recording only part of a run

diff --git a/Engine/Source/Developer/TraceServices/Private/Analyzers/NetTraceAnalyzer.cpp b/Engine/Source/Developer/TraceServices/Private/Analyzers/NetTraceAnalyzer.cpp
--- a/Engine/Source/Developer/TraceServices/Private/Analyzers/NetTraceAnalyzer.cpp
+++ b/Engine/Source/Developer/TraceServices/Private/Analyzers/NetTraceAnalyzer.cpp
@@ -52,6 +52,7 @@ void FNetTraceAnalyzer::OnAnalysisBegin(const FOnAnalysisContext& Context)
 	Builder.RouteEvent(RouteId_PacketStatsCounterEvent, "NetTrace", "PacketStatsCounterEvent");
 	Builder.RouteEvent(RouteId_FrameStatsCounterEvent, "NetTrace", "FrameStatsCounterEvent");
 	Builder.RouteEvent(RouteId_ObjectCreatedEvent, "NetTrace", "ObjectCreatedEvent");
+	Builder.RouteEvent(RouteId_ObjectExistsEvent, "NetTrace", "ObjectExistsEvent");
 	Builder.RouteEvent(RouteId_ObjectDestroyedEvent, "NetTrace", "ObjectDestroyedEvent");
 	Builder.RouteEvent(RouteId_ConnectionStateUpdatedEvent, "NetTrace", "ConnectionStateUpdatedEvent");
 	Builder.RouteEvent(RouteId_InstanceUpdatedEvent, "NetTrace", "InstanceUpdatedEvent");
@@ -104,7 +105,8 @@ bool FNetTraceAnalyzer::OnEvent(uint16 RouteId, EStyle Style, const FOnEventCont
 		case RouteId_InitEvent:
 		{
 			const uint64 TimestampCycles = EventData.GetValue<uint64>("Timestamp");
-			LastTimeStamp = Context.EventTime.AsSeconds(TimestampCycles);
+			StartTimeStamp = Context.EventTime.AsSeconds(TimestampCycles);
+			LastTimeStamp = StartTimeStamp;
 
 			// we always trace the version so that we make sure that we are backwards compatible with older trace stream
 			NetTraceVersion = EventData.GetValue<uint32>("NetTraceVersion");
@@ -197,6 +199,12 @@ bool FNetTraceAnalyzer::OnEvent(uint16 RouteId, EStyle Style, const FOnEventCont
 		}
 		break;
 
+		case RouteId_ObjectExistsEvent:
+		{
+			HandleObjectExistsEvent(Context, EventData);
+		}
+		break;
+
 		case RouteId_ObjectDestroyedEvent:
 		{
 			HandleObjectDestroyedEvent(Context, EventData);
@@ -858,6 +866,54 @@ void FNetTraceAnalyzer::HandleObjectCreatedEvent(const FOnEventContext& Context,
 	GameInstanceState->ActiveObjects.Add(ObjectId, { ObjectInstance.ObjectIndex, ObjectInstance.NameIndex });
 }
 
+void FNetTraceAnalyzer::HandleObjectExistsEvent(const FOnEventContext& Context, const FEventData& EventData)
+{
+	const uint64 TypeId = EventData.GetValue<uint64>("TypeId");
+	const uint64 ObjectId = EventData.GetValue<uint64>("ObjectId");
+	const FNetProfilerNameIndexType NameId = EventData.GetValue<FNetProfilerNameIndexType>("NameId");
+	const uint8 GameInstanceId = EventData.GetValue<uint8>("GameInstanceId");
+
+	TSharedRef<FNetTraceGameInstanceState> GameInstanceState = GetOrCreateActiveGameInstanceState(GameInstanceId);
+	const FNetProfilerNameIndexType* NetProfilerNameIndex = TracedNameIdToNetProfilerNameIdMap.Find(NameId);
+	const FNetProfilerNameIndexType NameIndex = NetProfilerNameIndex ? *NetProfilerNameIndex : 0u;
+
+	if (FNetTraceActiveObjectState* ActiveObjectInstance = GameInstanceState->ActiveObjects.Find(ObjectId))
+	{
+		if (FNetProfilerObjectInstance* ExistingInstance = NetProfilerProvider.EditObject(GameInstanceState->GameInstanceIndex, ActiveObjectInstance->ObjectIndex))
+		{
+			if (ExistingInstance->NameIndex == NameIndex || ExistingInstance->NameIndex == PendingNameIndex)
+			{
+				// Update existing object instance
+				ExistingInstance->LifeTime.Begin = StartTimeStamp;
+				ExistingInstance->NetObjectId = ObjectId;
+				ExistingInstance->TypeId = TypeId;
+
+				// Update name in both the persistent instance and the active one
+				ExistingInstance->NameIndex = NameIndex;
+				ActiveObjectInstance->NameIndex = NameIndex;
+
+				return;
+			}
+
+			// End instance and remove it from ActiveObjects
+			ExistingInstance->LifeTime.End = GetLastTimestamp();
+			GameInstanceState->ActiveObjects.Remove(ObjectId);
+		}
+	}
+
+	// Add persistent object representation
+	FNetProfilerObjectInstance& ObjectInstance = NetProfilerProvider.CreateObject(GameInstanceState->GameInstanceIndex);
+
+	// Fill in object data
+	ObjectInstance.LifeTime.Begin = StartTimeStamp;
+	ObjectInstance.NameIndex = NameIndex;
+	ObjectInstance.NetObjectId = ObjectId;
+	ObjectInstance.TypeId = TypeId;
+
+	// Add to active objects
+	GameInstanceState->ActiveObjects.Add(ObjectId, { ObjectInstance.ObjectIndex, ObjectInstance.NameIndex });
+}
+
 void FNetTraceAnalyzer::HandleObjectDestroyedEvent(const FOnEventContext& Context, const FEventData& EventData)
 {
 	// Remove from active instances and mark the end timestamp in the persistent instance list
diff --git a/Engine/Source/Developer/TraceServices/Private/Analyzers/NetTraceAnalyzer.h b/Engine/Source/Developer/TraceServices/Private/Analyzers/NetTraceAnalyzer.h
--- a/Engine/Source/Developer/TraceServices/Private/Analyzers/NetTraceAnalyzer.h
+++ b/Engine/Source/Developer/TraceServices/Private/Analyzers/NetTraceAnalyzer.h
@@ -40,6 +40,7 @@ private:
 		RouteId_PacketStatsCounterEvent,
 		RouteId_FrameStatsCounterEvent,
 		RouteId_ObjectCreatedEvent,
+		RouteId_ObjectExistsEvent,
 		RouteId_ObjectDestroyedEvent,
 		RouteId_ConnectionStateUpdatedEvent,
 		RouteId_InstanceUpdatedEvent,
@@ -149,6 +150,7 @@ private:
 	void HandleFrameStatsCounterEvent(const FOnEventContext& Context, const FEventData& EventData);
 	void HandlePacketStatsCounterEvent(const FOnEventContext& Context, const FEventData& EventData);
 	void HandleObjectCreatedEvent(const FOnEventContext& Context, const FEventData& EventData);
+	void HandleObjectExistsEvent(const FOnEventContext& Context, const FEventData& EventData);
 	void HandleObjectDestroyedEvent(const FOnEventContext& Context, const FEventData& EventData);
 
 	void HandleGameInstanceUpdatedEvent(const FOnEventContext& Context, const FEventData& EventData);
@@ -159,6 +161,7 @@ private:
 private:
 
 	TMap<uint32, TSharedRef<FNetTraceGameInstanceState>> ActiveGameInstances;
+	FNetProfilerTimeStamp StartTimeStamp;
 	FNetProfilerTimeStamp LastTimeStamp;
 };
 
diff --git a/Engine/Source/Runtime/Experimental/Iris/Core/Private/Iris/IrisCoreModule.cpp b/Engine/Source/Runtime/Experimental/Iris/Core/Private/Iris/IrisCoreModule.cpp
--- a/Engine/Source/Runtime/Experimental/Iris/Core/Private/Iris/IrisCoreModule.cpp
+++ b/Engine/Source/Runtime/Experimental/Iris/Core/Private/Iris/IrisCoreModule.cpp
@@ -8,11 +8,13 @@
 #include "Misc/CommandLine.h"
 
 #include "Net/Core/Connection/NetEnums.h"
+#include "Net/Core/Trace/Private/NetTraceInternal.h"
 
 #include "Iris/Core/IrisLog.h"
 #include "Iris/ReplicationState/DefaultPropertyNetSerializerInfos.h"
 #include "Iris/ReplicationState/PropertyNetSerializerInfoRegistry.h"
 #include "Iris/ReplicationSystem/LegacyPushModel.h"
+#include "Iris/ReplicationSystem/ReplicationOperations.h"
 #include "Iris/ReplicationSystem/ReplicationSystem.h"
 #include "Iris/Serialization/InternalNetSerializerDelegates.h"
 #include "Iris/IrisConfigInternal.h"
@@ -64,6 +66,13 @@ private:
 		{
 			RepSystemCount += (RepSystem != nullptr ? 1 : 0);
 		}
+
+#if UE_NET_TRACE_ENABLED
+		FNetTrace::OnResetPersistentNetDebugNames.AddLambda([]() 
+		{
+			UE::Net::ResetLifetimeConditionDebugNames();
+		});
+#endif
 	}
 
 	virtual void ShutdownModule() override
diff --git a/Engine/Source/Runtime/Experimental/Iris/Core/Private/Iris/ReplicationSystem/ReplicationOperations.cpp b/Engine/Source/Runtime/Experimental/Iris/Core/Private/Iris/ReplicationSystem/ReplicationOperations.cpp
--- a/Engine/Source/Runtime/Experimental/Iris/Core/Private/Iris/ReplicationSystem/ReplicationOperations.cpp
+++ b/Engine/Source/Runtime/Experimental/Iris/Core/Private/Iris/ReplicationSystem/ReplicationOperations.cpp
@@ -1394,6 +1394,14 @@ bool FReplicationProtocolOperations::IsEqualQuantizedState(FNetSerializationCont
 	return true;
 }
 
+void ResetLifetimeConditionDebugNames()
+{
+	for (FNetDebugName& Name : Private::s_LifetimeConditionDebugNames)
+	{
+		Name.DebugNameId = 0;
+	}
+}
+
 }
 
 namespace UE::Net::Private
diff --git a/Engine/Source/Runtime/Experimental/Iris/Core/Public/Iris/ReplicationSystem/ReplicationOperations.h b/Engine/Source/Runtime/Experimental/Iris/Core/Public/Iris/ReplicationSystem/ReplicationOperations.h
--- a/Engine/Source/Runtime/Experimental/Iris/Core/Public/Iris/ReplicationSystem/ReplicationOperations.h
+++ b/Engine/Source/Runtime/Experimental/Iris/Core/Public/Iris/ReplicationSystem/ReplicationOperations.h
@@ -162,4 +162,6 @@ struct FReplicationProtocolOperations
 	static IRISCORE_API void DeserializeWithMaskDelta(FNetSerializationContext& Context, uint32* DstChangeMaskData, uint8* RESTRICT DstObjectStateBuffer, const uint8* RESTRICT PrevObjectStateBuffer, const FReplicationProtocol* Protocol);
 };
 
+IRISCORE_API void ResetLifetimeConditionDebugNames();
+
 }
diff --git a/Engine/Source/Runtime/Net/Core/Private/Net/Core/Trace/NetTraceInternal.cpp b/Engine/Source/Runtime/Net/Core/Private/Net/Core/Trace/NetTraceInternal.cpp
--- a/Engine/Source/Runtime/Net/Core/Private/Net/Core/Trace/NetTraceInternal.cpp
+++ b/Engine/Source/Runtime/Net/Core/Private/Net/Core/Trace/NetTraceInternal.cpp
@@ -33,6 +33,9 @@ struct FNetTraceInternal
 
 	struct FThreadBuffer
 	{
+		// Trace Id allows us to discard the old Name Id maps when we start a new trace
+		uint32 TraceId;
+
 		// Map FName to NameId, used to know what FNames we have traced
 		TMap<FName, UE::Net::FNetDebugNameId> DynamicFNameToNameIdMap;
 
@@ -40,29 +43,106 @@ struct FNetTraceInternal
 		TMap<uint64, UE::Net::FNetDebugNameId> DynamicNameHashToNameIdMap;
 	};
 
+	struct FConnection
+	{
+		FString Address;
+		FString OwningActor;
+		uint8 State;
+	};
+
+	struct FNetObject
+	{
+		FName Name;
+		const UE::Net::FNetDebugName* DebugName;
+		uint64 TypeIdentifier; 
+		uint32 OwnerId;
+	};
+
+	struct FGameInstance
+	{
+		TMap<uint32, FConnection> Connections;
+		TMap<uint64, FNetObject> Objects;
+		TSet<uint64> ObjectsTraced;
+		FString Name;
+		bool bIsServer;
+	};
+
 	// Get next NameId used to track what we already have traced
-	UE::Net::FNetDebugNameId static GetNextNameId();
+	static UE::Net::FNetDebugNameId GetNextNameId();
 
-	static inline FThreadBuffer* CreateThreadBuffer();
+	static inline FThreadBuffer* GetThreadBuffer();
 
 	static constexpr ENetTraceVersion NetTraceVersion = ENetTraceVersion::ENetTraceVersion_DebugNameIndexIs32Bits;
 };
 
+static std::atomic<uint32> CurrentTraceId(0);
+static std::atomic<UE::Net::FNetDebugNameId> NextNameId(1);
 static thread_local TUniquePtr<FNetTraceInternal::FThreadBuffer> ThreadBuffer;
+static FDelegateHandle OnTraceStartedHandle;
+static FDelegateHandle OnTraceStoppedHandle;
+static TMap<uint32, FNetTraceInternal::FGameInstance> GameInstances;
 
-FNetTraceInternal::FThreadBuffer* FNetTraceInternal::CreateThreadBuffer()
+FNetTrace::FOnResetPersistentNetDebugNames FNetTrace::OnResetPersistentNetDebugNames;
+FNetTrace::FOnNetTraceStarted FNetTrace::OnNetTraceStarted;
+
+FNetTraceInternal::FThreadBuffer* FNetTraceInternal::GetThreadBuffer()
 {
-	ThreadBuffer = MakeUnique<FNetTraceInternal::FThreadBuffer>();
-	return ThreadBuffer.Get();
+	FNetTraceInternal::FThreadBuffer* ThreadBufferPtr = ThreadBuffer.Get();
+	if (!ThreadBufferPtr || ThreadBufferPtr->TraceId != CurrentTraceId)
+	{
+		ThreadBuffer = MakeUnique<FNetTraceInternal::FThreadBuffer>();
+		ThreadBufferPtr = ThreadBuffer.Get();
+		ThreadBufferPtr->TraceId = CurrentTraceId;
+	}
+	
+	return ThreadBufferPtr;
 }
 
 UE::Net::FNetDebugNameId FNetTraceInternal::GetNextNameId()
 {
-	static std::atomic<UE::Net::FNetDebugNameId> NextNameId(1);
-
 	return NextNameId++;
 }
 
+void FNetTrace::OnTraceStarted(FTraceAuxiliary::EConnectionType /*TraceType*/, const FString& /*TraceDestination*/)
+{
+	++CurrentTraceId;
+	NextNameId = 1;
+	UE::Net::ResetPersistentNetDebugNameIds();
+	OnResetPersistentNetDebugNames.Broadcast();
+
+	FNetTraceInternal::Reporter::ReportInitEvent(FNetTraceInternal::NetTraceVersion);
+
+	for (TPair<uint32, FNetTraceInternal::FGameInstance>& InstanceIter : GameInstances)
+	{
+		const uint32 GameInstanceId = InstanceIter.Key;
+		FNetTraceInternal::FGameInstance& GameInstance = InstanceIter.Value;
+
+		GameInstance.ObjectsTraced.Reset();
+
+		FNetTraceInternal::Reporter::ReportInstanceUpdated(GameInstanceId, GameInstance.bIsServer, *GameInstance.Name);
+
+		for (const TPair<uint32, FNetTraceInternal::FConnection>& ConnectionIter : GameInstance.Connections)
+		{
+			const uint32 ConnectionId = ConnectionIter.Key;
+			const FNetTraceInternal::FConnection& Connection = ConnectionIter.Value;
+
+			FNetTraceInternal::Reporter::ReportConnectionCreated(GameInstanceId, ConnectionId);
+			FNetTraceInternal::Reporter::ReportConnectionUpdated(GameInstanceId, ConnectionId, *Connection.Address, *Connection.OwningActor);
+			FNetTraceInternal::Reporter::ReportConnectionStateUpdated(GameInstanceId, ConnectionId, Connection.State);
+		};
+	}
+
+	OnNetTraceStarted.Broadcast();
+}
+
+void FNetTrace::OnTraceStopped(FTraceAuxiliary::EConnectionType /*TraceType*/, const FString& /*TraceDestination*/)
+{
+	if (ThreadBuffer)
+	{
+		ThreadBuffer.Reset();
+	}
+}
+
 void FNetTrace::SetTraceVerbosity(uint32 Verbosity)
 {
 	const uint32 NewVerbosity = FMath::Min(Verbosity, (uint32)UE_NET_TRACE_COMPILETIME_VERBOSITY);
@@ -73,16 +153,23 @@ void FNetTrace::SetTraceVerbosity(uint32 Verbosity)
 		UE::Trace::ToggleChannel(TEXT("NetChannel"), true);
 		UE::Trace::ToggleChannel(TEXT("FrameChannel"), true);
 
-		FNetTraceInternal::Reporter::ReportInitEvent(FNetTraceInternal::NetTraceVersion);
+		if (UE::Trace::IsTracing())
+		{
+			OnTraceStarted(FTraceAuxiliary::EConnectionType::None, TEXT(""));
+		}
+		OnTraceStartedHandle = FTraceAuxiliary::OnTraceStarted.AddStatic(&FNetTrace::OnTraceStarted);
+		OnTraceStoppedHandle = FTraceAuxiliary::OnTraceStopped.AddStatic(&FNetTrace::OnTraceStopped);
 	}
 	else if (GetTraceVerbosity() && !NewVerbosity)
 	{
-		if (ThreadBuffer)
+		UE::Trace::ToggleChannel(TEXT("NetChannel"), false);
+			
+		if (UE::Trace::IsTracing())
 		{
-			UE::Trace::ToggleChannel(TEXT("NetChannel"), false);
-
-			ThreadBuffer.Reset();
-		}		
+			OnTraceStopped(FTraceAuxiliary::EConnectionType::None, TEXT(""));
+		}
+		FTraceAuxiliary::OnTraceStarted.Remove(OnTraceStartedHandle);
+		FTraceAuxiliary::OnTraceStopped.Remove(OnTraceStoppedHandle);
 	}
 
 	GNetTraceRuntimeVerbosity = NewVerbosity;
@@ -90,6 +177,8 @@ void FNetTrace::SetTraceVerbosity(uint32 Verbosity)
 
 void FNetTrace::TraceEndSession(uint32 GameInstanceId)
 {
+	GameInstances.Remove(GameInstanceId);
+
 	if (GNetTraceRuntimeVerbosity)
 	{
 		FNetTraceInternal::Reporter::ReportInstanceDestroyed(GameInstanceId);
@@ -98,6 +187,10 @@ void FNetTrace::TraceEndSession(uint32 GameInstanceId)
 
 void FNetTrace::TraceInstanceUpdated(uint32 GameInstanceId, bool bIsServer, const TCHAR* Name)
 {
+	FNetTraceInternal::FGameInstance& Instance = GameInstances.FindOrAdd(GameInstanceId);
+	Instance.Name = Name;
+	Instance.bIsServer = bIsServer;
+
 	if (GNetTraceRuntimeVerbosity)
 	{
 		FNetTraceInternal::Reporter::ReportInstanceUpdated(GameInstanceId, bIsServer, Name);
@@ -158,6 +251,8 @@ void FNetTrace::FoldTraceCollector(FNetTraceCollector* DstCollector, const FNetT
 			}
 			
 			DstCollector->EventCount += SrcCollector->EventCount;
+
+			DstCollector->ReferencedObjectIds.Append(SrcCollector->ReferencedObjectIds);
 		}		
 	}
 }
@@ -383,6 +478,32 @@ void FNetTrace::TraceCollectedEvents(FNetTraceCollector& Collector, uint32 GameI
 	PacketInfo.PacketSequenceNumber = 0;
 	PacketInfo.PacketType = PacketType;
 
+	FNetTraceInternal::FGameInstance& GameInstance = GameInstances[GameInstanceId];
+	for(uint64 NetObjectId : Collector.ReferencedObjectIds)
+	{
+		if (!GameInstance.ObjectsTraced.Contains(NetObjectId))
+		{
+			const FNetTraceInternal::FNetObject* NetObject = GameInstance.Objects.Find(NetObjectId);
+			if (NetObject)
+			{
+				GameInstance.ObjectsTraced.Add(NetObjectId);
+
+				if (NetObject->DebugName)
+				{
+					if (NetObject->DebugName->DebugNameId == 0U)
+					{
+						NetObject->DebugName->DebugNameId = FNetTrace::TraceName(NetObject->DebugName->Name);
+					}
+					FNetTraceInternal::Reporter::ReportObjectExists(GameInstanceId, NetObjectId, NetObject->DebugName->DebugNameId, NetObject->TypeIdentifier, NetObject->OwnerId);
+				}
+				else
+				{
+					FNetTraceInternal::Reporter::ReportObjectExists(GameInstanceId, NetObjectId, FNetTrace::TraceName(NetObject->Name), NetObject->TypeIdentifier, NetObject->OwnerId);
+				}
+			}
+		}
+	}
+
 	// Trace all collected events
 	FNetTraceInternal::Reporter::ReportPacketContent(Collector.Events.GetData(), Collector.EventCount, PacketInfo);
 
@@ -421,11 +542,15 @@ void FNetTrace::TracePacket(uint32 GameInstanceId, uint32 ConnectionId, uint32 P
 
 void FNetTrace::TraceObjectCreated(uint32 GameInstanceId, uint64 NetObjectId, const UE::Net::FNetDebugName* DebugName, uint64 TypeIdentifier, uint32 OwnerId)
 {
+	GameInstances.FindOrAdd(GameInstanceId).Objects.Emplace(NetObjectId, { NAME_None, DebugName, TypeIdentifier, OwnerId });
+	
 	if (!GNetTraceRuntimeVerbosity)
 	{
 		return;
 	}
 
+	GameInstances[GameInstanceId].ObjectsTraced.Add(NetObjectId);
+
 	if (DebugName->DebugNameId == 0U)
 	{
 		DebugName->DebugNameId = FNetTrace::TraceName(DebugName->Name);
@@ -435,24 +560,35 @@ void FNetTrace::TraceObjectCreated(uint32 GameInstanceId, uint64 NetObjectId, co
 
 void FNetTrace::TraceObjectCreated(uint32 GameInstanceId, uint64 NetObjectId, const FName ObjectName, uint64 TypeIdentifier, uint32 OwnerId)
 {
+	GameInstances.FindOrAdd(GameInstanceId).Objects.Emplace(NetObjectId, {ObjectName, nullptr, TypeIdentifier, OwnerId});
+	
 	if (!GNetTraceRuntimeVerbosity)
 	{
 		return;
 	}
 
+	GameInstances[GameInstanceId].ObjectsTraced.Add(NetObjectId);
+
 	FNetTraceInternal::Reporter::ReportObjectCreated(GameInstanceId, NetObjectId, FNetTrace::TraceName(ObjectName), TypeIdentifier, OwnerId);
 }
 
 void FNetTrace::TraceObjectDestroyed(uint32 GameInstanceId, uint64 NetObjectId)
 {
+	GameInstances[GameInstanceId].Objects.Remove(NetObjectId);
+
 	if (GNetTraceRuntimeVerbosity)
 	{
+		GameInstances[GameInstanceId].ObjectsTraced.Remove(NetObjectId);
+
 		FNetTraceInternal::Reporter::ReportObjectDestroyed(GameInstanceId, NetObjectId);
 	}
 }
 
 void FNetTrace::TraceConnectionCreated(uint32 GameInstanceId, uint32 ConnectionId)
 {
+	// TraceConnectionStateUpdated can be called before this, so use FindOrAdd
+	GameInstances.FindChecked(GameInstanceId).Connections.FindOrAdd(ConnectionId);
+
 	if (GNetTraceRuntimeVerbosity)
 	{
 		FNetTraceInternal::Reporter::ReportConnectionCreated(GameInstanceId, ConnectionId);
@@ -461,6 +597,8 @@ void FNetTrace::TraceConnectionCreated(uint32 GameInstanceId, uint32 ConnectionI
 
 void FNetTrace::TraceConnectionStateUpdated(uint32 GameInstanceId, uint32 ConnectionId, uint8 ConnectionStateValue)
 {
+	GameInstances.FindChecked(GameInstanceId).Connections.FindOrAdd(ConnectionId).State = ConnectionStateValue;
+
 	if (GNetTraceRuntimeVerbosity)
 	{
 		FNetTraceInternal::Reporter::ReportConnectionStateUpdated(GameInstanceId, ConnectionId, ConnectionStateValue);
@@ -469,6 +607,10 @@ void FNetTrace::TraceConnectionStateUpdated(uint32 GameInstanceId, uint32 Connec
 
 void FNetTrace::TraceConnectionUpdated(uint32 GameInstanceId, uint32 ConnectionId, const TCHAR* AddressString, const TCHAR* OwningActor)
 {
+	FNetTraceInternal::FConnection& Connection = GameInstances.FindChecked(GameInstanceId).Connections.FindChecked(ConnectionId);
+	Connection.Address = AddressString;
+	Connection.OwningActor = OwningActor;
+
 	if (GNetTraceRuntimeVerbosity)
 	{
 		FNetTraceInternal::Reporter::ReportConnectionUpdated(GameInstanceId, ConnectionId, AddressString, OwningActor);
@@ -491,6 +633,14 @@ void FNetTrace::TracePacketStatsCounter(uint32 GameInstanceId, uint32 Connection
 	}
 }
 
+void FNetTrace::TracePacketStatsCounter(uint32 GameInstanceId, uint32 ConnectionId, const UE::Net::FNetDebugName* CounterName, uint32 StatValue)
+{
+	if (GNetTraceRuntimeVerbosity)
+	{
+		FNetTraceInternal::Reporter::ReportPacketStatsCounter(GameInstanceId, ConnectionId, CounterName->DebugNameId, StatValue);
+	}
+}
+
 void FNetTrace::TraceFrameStatsCounter(uint32 GameInstanceId, UE::Net::FNetDebugNameId CounterNameId, uint32 StatValue)
 {
 	if (GNetTraceRuntimeVerbosity)
@@ -499,6 +649,14 @@ void FNetTrace::TraceFrameStatsCounter(uint32 GameInstanceId, UE::Net::FNetDebug
 	}
 }
 
+void FNetTrace::TraceFrameStatsCounter(uint32 GameInstanceId, const UE::Net::FNetDebugName* CounterName, uint32 StatValue)
+{
+	if (GNetTraceRuntimeVerbosity)
+	{
+		FNetTraceInternal::Reporter::ReportFrameStatsCounter(GameInstanceId, CounterName->DebugNameId, StatValue);
+	}
+}
+
 UE::Net::FNetDebugNameId FNetTrace::TraceName(const TCHAR* Name)
 {
 	if ((GNetTraceRuntimeVerbosity == 0U) | (Name == nullptr))
@@ -507,11 +665,7 @@ UE::Net::FNetDebugNameId FNetTrace::TraceName(const TCHAR* Name)
 	}
 
 	// Get Thread buffer
-	FNetTraceInternal::FThreadBuffer* ThreadBufferPtr = ThreadBuffer.Get();
-	if (!ThreadBufferPtr)
-	{
-		ThreadBufferPtr = FNetTraceInternal::CreateThreadBuffer();
-	}
+	FNetTraceInternal::FThreadBuffer* ThreadBufferPtr = FNetTraceInternal::GetThreadBuffer();
 
 	// Hash the name using CityHash64
 	const uint64 HashedName = CityHash64((const char*)Name, FCString::Strlen(Name) * sizeof(TCHAR));
@@ -541,11 +695,7 @@ UE::Net::FNetDebugNameId FNetTrace::TraceName(FName Name)
 	}
 
 	// Get Thread buffer
-	FNetTraceInternal::FThreadBuffer* ThreadBufferPtr = ThreadBuffer.Get();
-	if (!ThreadBufferPtr)
-	{
-		ThreadBufferPtr = FNetTraceInternal::CreateThreadBuffer();
-	}
+	FNetTraceInternal::FThreadBuffer* ThreadBufferPtr = FNetTraceInternal::GetThreadBuffer();
 
 	if (const FNetDebugNameId* FoundNameId = ThreadBufferPtr->DynamicFNameToNameIdMap.Find(Name))
 	{
diff --git a/Engine/Source/Runtime/Net/Core/Private/Net/Core/Trace/PersistentNetDebugNames.cpp b/Engine/Source/Runtime/Net/Core/Private/Net/Core/Trace/PersistentNetDebugNames.cpp
--- a/Engine/Source/Runtime/Net/Core/Private/Net/Core/Trace/PersistentNetDebugNames.cpp
+++ b/Engine/Source/Runtime/Net/Core/Private/Net/Core/Trace/PersistentNetDebugNames.cpp
@@ -21,6 +21,7 @@ public:
 	bool CreatePersistentNetDebugName(const TCHAR* Name, uint32 Length, const FNetDebugName*& OutDebugName);
 
 	static FPersistentNetDebugNames& Get();
+	static void ResetNameIds();
 	static void TearDown();
 private:
 
@@ -50,6 +51,19 @@ FPersistentNetDebugNames& FPersistentNetDebugNames::Get()
 	return Singleton;
 }
 
+void FPersistentNetDebugNames::ResetNameIds()
+{
+	FPersistentNetDebugNames& Singleton = Get();
+	{
+		UE::TScopeLock Lock(Singleton.CriticalSection);
+
+		for (const auto& Iter : Singleton.PersistentNames)
+		{
+			Iter.Value->DebugNameId = 0;
+		}
+	}
+}
+
 void FPersistentNetDebugNames::TearDown()
 {
 	FPersistentNetDebugNames& Singleton = Get();	
@@ -101,4 +115,9 @@ const FNetDebugName* CreatePersistentNetDebugName(const TCHAR* Name, uint32 Name
 	return DebugName;
 }
 
+void ResetPersistentNetDebugNameIds()
+{
+	FPersistentNetDebugNames::Get().ResetNameIds();
+}
+
 }
diff --git a/Engine/Source/Runtime/Net/Core/Private/Net/Core/Trace/Reporters/NetTraceReporter.cpp b/Engine/Source/Runtime/Net/Core/Private/Net/Core/Trace/Reporters/NetTraceReporter.cpp
--- a/Engine/Source/Runtime/Net/Core/Private/Net/Core/Trace/Reporters/NetTraceReporter.cpp
+++ b/Engine/Source/Runtime/Net/Core/Private/Net/Core/Trace/Reporters/NetTraceReporter.cpp
@@ -37,6 +37,14 @@ UE_TRACE_EVENT_BEGIN(NetTrace, ObjectCreatedEvent)
 	UE_TRACE_EVENT_FIELD(uint8, GameInstanceId)
 UE_TRACE_EVENT_END()
 
+UE_TRACE_EVENT_BEGIN(NetTrace, ObjectExistsEvent)	
+	UE_TRACE_EVENT_FIELD(uint64, TypeId)
+	UE_TRACE_EVENT_FIELD(uint64, ObjectId)
+	UE_TRACE_EVENT_FIELD(uint32, OwnerId)
+	UE_TRACE_EVENT_FIELD(uint32, NameId)
+	UE_TRACE_EVENT_FIELD(uint8, GameInstanceId)
+UE_TRACE_EVENT_END()
+
 UE_TRACE_EVENT_BEGIN(NetTrace, ObjectDestroyedEvent)
 	UE_TRACE_EVENT_FIELD(uint64, ObjectId)
 	UE_TRACE_EVENT_FIELD(uint8, GameInstanceId)
@@ -345,6 +353,16 @@ void FNetTraceReporter::ReportObjectCreated(uint32 GameInstanceId, uint64 NetObj
 		<< ObjectCreatedEvent.GameInstanceId((uint8)GameInstanceId);
 }
 
+void FNetTraceReporter::ReportObjectExists(uint32 GameInstanceId, uint64 NetObjectId, UE::Net::FNetDebugNameId NameId, uint64 TypeIdentifier, uint32 OwnerId)
+{
+	UE_TRACE_LOG(NetTrace, ObjectExistsEvent, NetChannel)
+		<< ObjectExistsEvent.TypeId(TypeIdentifier)
+		<< ObjectExistsEvent.ObjectId(NetObjectId)
+		<< ObjectExistsEvent.OwnerId(OwnerId)
+		<< ObjectExistsEvent.NameId(NameId)
+		<< ObjectExistsEvent.GameInstanceId((uint8)GameInstanceId);
+}
+
 void FNetTraceReporter::ReportObjectDestroyed(uint32 GameInstanceId, uint64 NetObjectId)
 {
 	UE_TRACE_LOG(NetTrace, ObjectDestroyedEvent, NetChannel)
diff --git a/Engine/Source/Runtime/Net/Core/Private/Net/Core/Trace/Reporters/NetTraceReporter.h b/Engine/Source/Runtime/Net/Core/Private/Net/Core/Trace/Reporters/NetTraceReporter.h
--- a/Engine/Source/Runtime/Net/Core/Private/Net/Core/Trace/Reporters/NetTraceReporter.h
+++ b/Engine/Source/Runtime/Net/Core/Private/Net/Core/Trace/Reporters/NetTraceReporter.h
@@ -16,6 +16,7 @@ struct FNetTraceReporter
 	static void ReportPacket(const FNetTracePacketInfo& PacketInfo, uint32 PacketBits);
 	static void ReportAnsiName(UE::Net::FNetDebugNameId NameId, uint32 NameSize, const char* Name);	
 	static void ReportObjectCreated(uint32 GameInstanceId, uint64 NetObjectId, UE::Net::FNetDebugNameId NameId, uint64 TypeIdentifier, uint32 OwnerId);
+	static void ReportObjectExists(uint32 GameInstanceId, uint64 NetObjectId, UE::Net::FNetDebugNameId NameId, uint64 TypeIdentifier, uint32 OwnerId);
 	static void ReportObjectDestroyed(uint32 GameInstanceId, uint64 NetObjectId);
 	static void ReportConnectionCreated(uint32 GameInstanceId, uint32 ConnectionId);
 	static void ReportConnectionStateUpdated(uint32 GameInstanceId, uint32 ConnectionId, uint8 ConnectionStateValue);
diff --git a/Engine/Source/Runtime/Net/Core/Public/Net/Core/Trace/NetDebugName.h b/Engine/Source/Runtime/Net/Core/Public/Net/Core/Trace/NetDebugName.h
--- a/Engine/Source/Runtime/Net/Core/Public/Net/Core/Trace/NetDebugName.h
+++ b/Engine/Source/Runtime/Net/Core/Public/Net/Core/Trace/NetDebugName.h
@@ -23,6 +23,8 @@ struct FNetDebugName
 NETCORE_API const FNetDebugName* CreatePersistentNetDebugName(const TCHAR* Name, uint32 NameLen);
 inline const FNetDebugName* CreatePersistentNetDebugName(const TCHAR* Name) { return CreatePersistentNetDebugName(Name, FCString::Strlen(Name)); }
 
+NETCORE_API void ResetPersistentNetDebugNameIds();
+
 }
 
 inline const TCHAR* ToCStr(const UE::Net::FNetDebugName* DebugName)
diff --git a/Engine/Source/Runtime/Net/Core/Public/Net/Core/Trace/Private/NetTraceInternal.h b/Engine/Source/Runtime/Net/Core/Public/Net/Core/Trace/Private/NetTraceInternal.h
--- a/Engine/Source/Runtime/Net/Core/Public/Net/Core/Trace/Private/NetTraceInternal.h
+++ b/Engine/Source/Runtime/Net/Core/Public/Net/Core/Trace/Private/NetTraceInternal.h
@@ -13,6 +13,7 @@
 #include "Misc/NetworkGuid.h"
 #include "Net/Core/Trace/NetDebugName.h"
 #include "UObject/NameTypes.h"
+#include "ProfilingDebugging/TraceAuxiliary.h"
 
 class FNetTraceCollector;
 
@@ -166,6 +167,16 @@ struct FNetTrace
 	/** Sets the runtime trace verbosity level, Verbosity == 0 means that it is disabled */
 	NETCORE_API static void SetTraceVerbosity(uint32 Verbosity);
 
+	static void OnTraceStarted(FTraceAuxiliary::EConnectionType TraceType, const FString& TraceDestination);
+
+	static void OnTraceStopped(FTraceAuxiliary::EConnectionType TraceType, const FString& TraceDestination);
+
+	DECLARE_TS_MULTICAST_DELEGATE(FOnResetPersistentNetDebugNames);
+	static NETCORE_API FOnResetPersistentNetDebugNames OnResetPersistentNetDebugNames;
+
+	DECLARE_TS_MULTICAST_DELEGATE(FOnNetTraceStarted);
+	static NETCORE_API FOnNetTraceStarted OnNetTraceStarted;
+
 	/** Trace end of session for a GameInstance */
 	NETCORE_API static void TraceEndSession(uint32 GameInstanceId);
 
@@ -253,8 +264,10 @@ struct FNetTrace
 
 	/** Trace stats counters for networking */
 	NETCORE_API static void TracePacketStatsCounter(uint32 GameInstanceId, uint32 ConnectionId, UE::Net::FNetDebugNameId CounterNameId, uint32 StatValue);
+	NETCORE_API static void TracePacketStatsCounter(uint32 GameInstanceId, uint32 ConnectionId, const UE::Net::FNetDebugName* CounterName, uint32 StatValue);
 	NETCORE_API static void TraceFrameStatsCounter(uint32 GameInstanceId, UE::Net::FNetDebugNameId CounterNameId,  uint32 StatValue);
-
+	NETCORE_API static void TraceFrameStatsCounter(uint32 GameInstanceId, const UE::Net::FNetDebugName* CounterNameId,  uint32 StatValue);
+	
 	/** Trace the name */
 	NETCORE_API static UE::Net::FNetDebugNameId TraceName(const TCHAR* Name);
 
@@ -323,6 +336,7 @@ public:
 	// constants
 	enum { MaxNestingLevel = 32 };
 	enum { InlinedEventCount = 256 };
+	enum { InlinedObjectIdCount = 128 };
 
 	explicit FNetTraceCollector(uint32 InitialEventCount) { Events.SetNumUninitialized(InitialEventCount); Reset(); }
 	FNetTraceCollector() : FNetTraceCollector(InlinedEventCount) {}; 
@@ -331,6 +345,7 @@ public:
 
 	inline FNetTracePacketContentEvent& EditEvent(uint32 Index) { return Events[Index]; }
 	inline FNetTracePacketContentEvent& EditEventNoCheck(uint32 Index) { return Events.GetData()[Index]; }
+	inline void AddReferencedObjectId(uint64 ObjectId) { ReferencedObjectIds.Add(ObjectId); }
 	
 private:
 	friend struct FNetTrace;
@@ -350,6 +365,9 @@ private:
 	// Events
 	typedef TArray<FNetTracePacketContentEvent, TInlineAllocator<InlinedEventCount> > FEventArray;
 	FEventArray Events;
+
+	// Object Ids referenced by events
+	TArray<uint64, TInlineAllocator<InlinedObjectIdCount>> ReferencedObjectIds;
 };
 
 /**
@@ -420,6 +438,10 @@ FNetTraceEventScope<T>::FNetTraceEventScope(uint64 InObjectId, T& InStream, FNet
 	{
 		EventIndex = FNetTrace::BeginPacketContentEvent(*Collector, ENetTracePacketContentEventType::Object, GetBitStreamPositionForNetTrace(Stream));
 		Collector->EditEventNoCheck(EventIndex).ObjectId = InObjectId;
+		if (InObjectId != 0)
+		{
+			Collector->AddReferencedObjectId(InObjectId);
+		}
 	}
 }
 
@@ -471,6 +493,7 @@ void FNetTraceEventScope<T>::SetEventObjectId(uint64 ObjectId)
 	if (EventIndex != FNetTrace::InvalidEventIndex)
 	{
 		Collector->EditEventNoCheck(EventIndex).ObjectId = ObjectId;
+		Collector->AddReferencedObjectId(ObjectId);
 	}
 }
 
@@ -532,6 +555,7 @@ void FNetTraceCollector::Reset()
 	PendingBunchEventIndex = ~0U;
 	OffsetStackLevel = 0U;
 	OffsetStack[0] = 0U;
+	ReferencedObjectIds.Reset();
 }
 
 #define UE_NET_TRACE_DO_IF(Cond, x) do { if (Cond) { x; } } while (0)
@@ -548,7 +572,7 @@ void FNetTraceCollector::Reset()
 #define UE_NET_TRACE_INTERNAL_BUNCH_SCOPE(Collector, Bunch, ...) FNetTraceBunchScope<decltype(Bunch)> PREPROCESSOR_JOIN(NetTraceBunchScope, __LINE__)(Bunch, __VA_ARGS__, Collector)
 
 #define UE_NET_TRACE_INTERNAL_SCOPE(Name, Stream, Collector, Verbosity) \
-	auto PREPROCESSOR_JOIN(NetTraceNameFunc_, __LINE__) = []() { static UE::Net::FNetDebugNameId NameId = FNetTrace::TraceName(TEXT(#Name)); return NameId; }; \
+	auto PREPROCESSOR_JOIN(NetTraceNameFunc_, __LINE__) = []() { static const UE::Net::FNetDebugName* Name = UE::Net::CreatePersistentNetDebugName(TEXT(#Name)); return FNetTrace::TraceName(Name); }; \
 	FNetTrace::TChooseTraceEventScope<Verbosity, decltype(Stream)>::Type PREPROCESSOR_JOIN(NetTraceScope, __LINE__)(PREPROCESSOR_JOIN(NetTraceNameFunc_, __LINE__), Stream, Collector, Verbosity)
 
 #define UE_NET_TRACE_INTERNAL_OBJECT_SCOPE(HandleOrNetGUID, Stream, Collector, Verbosity) \
@@ -559,7 +583,7 @@ void FNetTraceCollector::Reset()
 	FNetTrace::TChooseTraceEventScope<Verbosity, decltype(Stream)>::Type PREPROCESSOR_JOIN(NetTraceScope, __LINE__)(PREPROCESSOR_JOIN(NetTraceDynamicNameFunc_, __LINE__), Stream, Collector, Verbosity)
 
 #define UE_NET_TRACE_INTERNAL_NAMED_SCOPE(ScopeName, EventName, Stream, Collector, Verbosity) \
-	auto PREPROCESSOR_JOIN(NetTraceNameFunc_, __LINE__) = []() { static UE::Net::FNetDebugNameId  NameId = FNetTrace::TraceName(TEXT(#EventName)); return NameId; }; \
+	auto PREPROCESSOR_JOIN(NetTraceNameFunc_, __LINE__) = []() { static const UE::Net::FNetDebugName* Name = UE::Net::CreatePersistentNetDebugName(TEXT(#EventName)); return FNetTrace::TraceName(Name); }; \
 	FNetTrace::TChooseTraceEventScope<Verbosity, decltype(Stream)>::Type ScopeName(PREPROCESSOR_JOIN(NetTraceNameFunc_, __LINE__), Stream, Collector, Verbosity)
 
 #define UE_NET_TRACE_INTERNAL_NAMED_OBJECT_SCOPE(ScopeName, HandleOrNetGUID, Stream, Collector, Verbosity) \
@@ -576,8 +600,8 @@ void FNetTraceCollector::Reset()
 
 #define UE_NET_TRACE_INTERNAL(Name, Collector, StartPos, EndPos, Verbosity) \
 	do { if (FNetTrace::GetCollectorAtVerbosity(Collector, Verbosity)) { \
-			static UE::Net::FNetDebugNameId PREPROCESSOR_JOIN(NetTraceNameId_, __LINE__) = FNetTrace::TraceName(TEXT(#Name)); \
-			FNetTrace::TracePacketContentEvent(*Collector, PREPROCESSOR_JOIN(NetTraceNameId_, __LINE__), StartPos, EndPos, Verbosity); \
+			static const UE::Net::FNetDebugName* PREPROCESSOR_JOIN(NetTraceNameId_, __LINE__) = UE::Net::CreatePersistentNetDebugName(TEXT(#Name)); \
+			FNetTrace::TracePacketContentEvent(*Collector, FNetTrace::TraceName(PREPROCESSOR_JOIN(NetTraceNameId_, __LINE__)), StartPos, EndPos, Verbosity); \
 	} } while (0)
 
 #define UE_NET_TRACE_INTERNAL_DYNAMIC_NAME(Name, Collector, StartPos, EndPos, Verbosity) UE_NET_TRACE_DO_IF(FNetTrace::GetCollectorAtVerbosity(Collector, Verbosity), FNetTrace::TracePacketContentEvent(*Collector, FNetTrace::TraceName(Name), StartPos, EndPos, Verbosity))
@@ -594,13 +618,13 @@ void FNetTraceCollector::Reset()
 
 #define UE_NET_TRACE_INTERNAL_PACKET_STATSCOUNTER(GameInstanceId, ConnectionId, Name, StatValue, Verbosity) \
 	do { if (FNetTrace::GetNetTraceVerbosityEnabled(Verbosity)) { \
-			static UE::Net::FNetDebugNameId PREPROCESSOR_JOIN(NetTraceNameId_, __LINE__) = FNetTrace::TraceName(TEXT(#Name)); \
+			static const UE::Net::FNetDebugName* PREPROCESSOR_JOIN(NetTraceNameId_, __LINE__) = UE::Net::CreatePersistentNetDebugName(TEXT(#Name)); \
 			FNetTrace::TracePacketStatsCounter(GameInstanceId, ConnectionId, PREPROCESSOR_JOIN(NetTraceNameId_, __LINE__), StatValue); \
 	} } while (0)
 
 #define UE_NET_TRACE_INTERNAL_FRAME_STATSCOUNTER(GameInstanceId, Name, StatValue, Verbosity) \
 	do { if (FNetTrace::GetNetTraceVerbosityEnabled(Verbosity)) { \
-			static UE::Net::FNetDebugNameId PREPROCESSOR_JOIN(NetTraceNameId_, __LINE__) = FNetTrace::TraceName(TEXT(#Name)); \
+			static const UE::Net::FNetDebugName* PREPROCESSOR_JOIN(NetTraceNameId_, __LINE__) = UE::Net::CreatePersistentNetDebugName(TEXT(#Name)); \
 			FNetTrace::TraceFrameStatsCounter(GameInstanceId, PREPROCESSOR_JOIN(NetTraceNameId_, __LINE__), StatValue); \
 	} } while (0)
 
