Disabling ticking and user callbacks for more engine systems.

diff --git a/Engine/Source/Runtime/CoreUObject/Public/UObject/CoreNet.h b/Engine/Source/Runtime/CoreUObject/Public/UObject/CoreNet.h
--- a/Engine/Source/Runtime/CoreUObject/Public/UObject/CoreNet.h
+++ b/Engine/Source/Runtime/CoreUObject/Public/UObject/CoreNet.h
@@ -695,6 +695,11 @@ struct FNetDeltaSerializeInfo
 
 	/** Whether or not we support FFastArraySerializer::FastArrayDeltaSerialize_DeltaSerializeStructs */
 	bool bSupportsFastArrayDeltaStructSerialization = false;
+	
+#if UE_SUPPORT_FOR_ACTOR_TICK_DISABLE
+	/** Whether or not FFastArraySerializer will call user provided client-side callbacks (e.g. PreReplicatedRemove()). */
+	bool bCallFastArrayClientCallbacks = true;
+#endif
 
 	/**
 	 * Whether or the connection is completely reliable.
diff --git a/Engine/Source/Runtime/Engine/Private/DataReplication.cpp b/Engine/Source/Runtime/Engine/Private/DataReplication.cpp
--- a/Engine/Source/Runtime/Engine/Private/DataReplication.cpp
+++ b/Engine/Source/Runtime/Engine/Private/DataReplication.cpp
@@ -365,6 +365,12 @@ bool FObjectReplicator::SendCustomDeltaProperty(UObject* InObject, uint16 Custom
 	Parms.CustomDeltaIndex = CustomDeltaIndex;
 	Parms.bSupportsFastArrayDeltaStructSerialization = bSupportsFastArrayDelta;
 	Parms.Connection = Connection;
+#if UE_SUPPORT_FOR_ACTOR_TICK_DISABLE
+	if (ConnectionDriver && ConnectionDriver->GetWorld())
+	{
+		Parms.bCallFastArrayClientCallbacks = ConnectionDriver->GetWorld()->IsActorTickAndUserCallbacksEnabled();
+	}
+#endif
 	Parms.bInternalAck = Connection->IsInternalAck();
 
 	// When initializing baselines we should not modify the source data if it originates from the CDO or archetype
@@ -1130,6 +1136,12 @@ bool FObjectReplicator::ReceivedBunch(FNetBitReader& Bunch, const FReplicationFl
 			Parms.Reader = &Reader;
 			Parms.NetSerializeCB = &NetSerializeCB;
 			Parms.Connection = Connection;
+#if UE_SUPPORT_FOR_ACTOR_TICK_DISABLE
+			if (ConnectionNetDriver && ConnectionNetDriver->GetWorld())
+			{
+				Parms.bCallFastArrayClientCallbacks = ConnectionNetDriver->GetWorld()->IsActorTickAndUserCallbacksEnabled();
+			}
+#endif
 			Parms.bInternalAck = Connection->IsInternalAck();
 			Parms.Object = Object;
 
@@ -1524,6 +1536,16 @@ bool FObjectReplicator::MoveMappedObjectToUnmapped(const FNetworkGUID& GUID)
 	FNetSerializeCB NetSerializeCB(Connection->Driver);
 	FNetDeltaSerializeInfo Parms;
 	Parms.Connection = Connection;
+#if UE_SUPPORT_FOR_ACTOR_TICK_DISABLE
+	if (Connection)
+	{
+		UNetDriver* ConnectionDriver = Connection->GetDriver();
+		if (ConnectionDriver && ConnectionDriver->GetWorld())
+		{
+			Parms.bCallFastArrayClientCallbacks = ConnectionDriver->GetWorld()->IsActorTickAndUserCallbacksEnabled();
+		}
+	}
+#endif
 	Parms.bInternalAck = Connection->IsInternalAck();
 	Parms.Map = Connection->PackageMap;
 	Parms.Object = GetObject();
@@ -2426,6 +2448,16 @@ void FObjectReplicator::UpdateUnmappedObjects(bool & bOutHasMoreUnmapped)
 	FNetDeltaSerializeInfo Parms;
 	Parms.Object = Object;
 	Parms.Connection = Connection;
+#if UE_SUPPORT_FOR_ACTOR_TICK_DISABLE
+	if (Connection)
+	{
+		UNetDriver* ConnectionDriver = Connection->GetDriver();
+		if (ConnectionDriver && ConnectionDriver->GetWorld())
+		{
+			Parms.bCallFastArrayClientCallbacks = ConnectionDriver->GetWorld()->IsActorTickAndUserCallbacksEnabled();
+		}
+	}
+#endif
 	Parms.bInternalAck = Connection->IsInternalAck();
 	Parms.Map = Connection->PackageMap;
 	Parms.NetSerializeCB = &NetSerializeCB;
diff --git a/Engine/Source/Runtime/Engine/Private/LevelTick.cpp b/Engine/Source/Runtime/Engine/Private/LevelTick.cpp
--- a/Engine/Source/Runtime/Engine/Private/LevelTick.cpp
+++ b/Engine/Source/Runtime/Engine/Private/LevelTick.cpp
@@ -1564,6 +1564,9 @@ void UWorld::Tick( ELevelTick TickType, float DeltaSeconds )
 			}
 
 			// Update cameras
+#if UE_SUPPORT_FOR_ACTOR_TICK_DISABLE
+			if (IsActorTickAndUserCallbacksEnabled())
+#endif
 			{
 				SCOPE_CYCLE_COUNTER(STAT_UpdateCameraTime);
 				CSV_SCOPED_TIMING_STAT_EXCLUSIVE(Camera);
@@ -1585,6 +1588,9 @@ void UWorld::Tick( ELevelTick TickType, float DeltaSeconds )
 			}
 
 			// Update streaming volumes
+#if UE_SUPPORT_FOR_ACTOR_TICK_DISABLE
+			if (IsActorTickAndUserCallbacksEnabled())
+#endif
 			{
 				if (!bIsPaused && IsGameWorld())
 				{
diff --git a/Engine/Source/Runtime/Net/Core/Classes/Net/Serialization/FastArraySerializer.h b/Engine/Source/Runtime/Net/Core/Classes/Net/Serialization/FastArraySerializer.h
--- a/Engine/Source/Runtime/Net/Core/Classes/Net/Serialization/FastArraySerializer.h
+++ b/Engine/Source/Runtime/Net/Core/Classes/Net/Serialization/FastArraySerializer.h
@@ -1135,10 +1135,21 @@ void FFastArraySerializer::TFastArraySerializeHelper<Type, SerializerType>::Post
 			}
 
 			// Call the delete callbacks now, actually remove them at the end
-			Items[idx].PreReplicatedRemove(ArraySerializer);
+#if UE_SUPPORT_FOR_ACTOR_TICK_DISABLE
+			if (Parms.bCallFastArrayClientCallbacks)
+#endif
+			{
+				Items[idx].PreReplicatedRemove(ArraySerializer);
+			}
 		}
 	}
-	ArraySerializer.PreReplicatedRemove(Header.DeletedIndices, FinalSize);
+
+#if UE_SUPPORT_FOR_ACTOR_TICK_DISABLE
+	if (Parms.bCallFastArrayClientCallbacks)
+#endif
+	{
+		ArraySerializer.PreReplicatedRemove(Header.DeletedIndices, FinalSize);
+	}
 
 	if (PreRemoveSize != Items.Num())
 	{
@@ -1146,17 +1157,27 @@ void FFastArraySerializer::TFastArraySerializeHelper<Type, SerializerType>::Post
 			PreRemoveSize, Items.Num());
 	}
 
-	for (int32 idx : AddedIndices)
+#if UE_SUPPORT_FOR_ACTOR_TICK_DISABLE
+	if (Parms.bCallFastArrayClientCallbacks)
+#endif
 	{
-		Items[idx].PostReplicatedAdd(ArraySerializer);
+		for (int32 idx : AddedIndices)
+		{
+			Items[idx].PostReplicatedAdd(ArraySerializer);
+		}
+		ArraySerializer.PostReplicatedAdd(AddedIndices, FinalSize);
 	}
-	ArraySerializer.PostReplicatedAdd(AddedIndices, FinalSize);
 
-	for (int32 idx : ChangedIndices)
+#if UE_SUPPORT_FOR_ACTOR_TICK_DISABLE
+	if (Parms.bCallFastArrayClientCallbacks)
+#endif
 	{
-		Items[idx].PostReplicatedChange(ArraySerializer);
+		for (int32 idx : ChangedIndices)
+		{
+			Items[idx].PostReplicatedChange(ArraySerializer);
+		}
+		ArraySerializer.PostReplicatedChange(ChangedIndices, FinalSize);
 	}
-	ArraySerializer.PostReplicatedChange(ChangedIndices, FinalSize);
 
 	if (PreRemoveSize != Items.Num())
 	{
@@ -1336,8 +1357,13 @@ bool FFastArraySerializer::FastArrayDeltaSerialize(TArray<Type> &Items, FNetDelt
 				Parms.Reader = &Reader;
 				Parms.NetSerializeCB->NetSerializeStruct(Parms);
 
-				// Let the element know it changed
-				ThisElement->PostReplicatedChange(ArraySerializer);
+#if UE_SUPPORT_FOR_ACTOR_TICK_DISABLE
+				if (Parms.bCallFastArrayClientCallbacks)
+#endif
+				{
+					// Let the element know it changed
+					ThisElement->PostReplicatedChange(ArraySerializer);
+				}
 			}
 
 			// If we have no more guids, we can remove this item for good
@@ -1355,8 +1381,13 @@ bool FFastArraySerializer::FastArrayDeltaSerialize(TArray<Type> &Items, FNetDelt
 
 		if (Parms.bOutSomeObjectsWereMapped)
 		{
-			ArraySerializer.PostReplicatedChange(ChangedIndices, Items.Num());
-			Helper.CallPostReplicatedReceiveOrNot(Items.Num());
+#if UE_SUPPORT_FOR_ACTOR_TICK_DISABLE
+			if (Parms.bCallFastArrayClientCallbacks)
+#endif
+			{
+				ArraySerializer.PostReplicatedChange(ChangedIndices, Items.Num());
+				Helper.CallPostReplicatedReceiveOrNot(Items.Num());
+			}
 		}
 		return true;
 	}
@@ -1584,7 +1615,12 @@ bool FFastArraySerializer::FastArrayDeltaSerialize(TArray<Type> &Items, FNetDelt
 
 		Helper.template PostReceiveCleanup<>(Header, ChangedIndices, AddedIndices, ArraySerializer.GuidReferencesMap);
 
-		Helper.CallPostReplicatedReceiveOrNot(OldNumItems);
+#if UE_SUPPORT_FOR_ACTOR_TICK_DISABLE
+		if (Parms.bCallFastArrayClientCallbacks)
+#endif
+		{
+			Helper.CallPostReplicatedReceiveOrNot(OldNumItems);
+		}
 	}
 
 	return true;
@@ -1693,10 +1729,15 @@ bool FFastArraySerializer::FastArrayDeltaSerialize_DeltaSerializeStructs(TArray<
 
 		Parms.NetSerializeCB->UpdateUnmappedGuidsForFastArray(DeltaSerializeParams);
 
-		if (Parms.bOutSomeObjectsWereMapped)
+#if UE_SUPPORT_FOR_ACTOR_TICK_DISABLE
+		if (Parms.bCallFastArrayClientCallbacks)
+#endif
 		{
-			ArraySerializer.PostReplicatedChange(ChangedIndices, Items.Num());
-			Helper.CallPostReplicatedReceiveOrNot(Items.Num());
+			if (Parms.bOutSomeObjectsWereMapped)
+			{
+				ArraySerializer.PostReplicatedChange(ChangedIndices, Items.Num());
+				Helper.CallPostReplicatedReceiveOrNot(Items.Num());
+			}
 		}
 
 		return true;
@@ -1819,7 +1860,12 @@ bool FFastArraySerializer::FastArrayDeltaSerialize_DeltaSerializeStructs(TArray<
 
 		Helper.template PostReceiveCleanup<>(Header, ChangedIndices, AddedIndices, ArraySerializer.GuidReferencesMap_StructDelta);
 
-		Helper.CallPostReplicatedReceiveOrNot(OldNumItems);
+#if UE_SUPPORT_FOR_ACTOR_TICK_DISABLE
+		if (Parms.bCallFastArrayClientCallbacks)
+#endif
+		{
+			Helper.CallPostReplicatedReceiveOrNot(OldNumItems);
+		}
 	}
 
 	return true;
