[AutoRTFM] Update atomic accesses in FNetTrace for transactional safety.

We don't need to roll back `NextNameId` if a transaction aborts, so it's safe to access and increment this atomic from the open.
`GetThreadBuffer()` doesn't modify `CurrentTraceId` at all, so it is safe to read the atomic from the open.

diff --git a/Engine/Source/Runtime/Net/Core/Private/Net/Core/Trace/NetTraceInternal.cpp b/Engine/Source/Runtime/Net/Core/Private/Net/Core/Trace/NetTraceInternal.cpp
--- a/Engine/Source/Runtime/Net/Core/Private/Net/Core/Trace/NetTraceInternal.cpp
+++ b/Engine/Source/Runtime/Net/Core/Private/Net/Core/Trace/NetTraceInternal.cpp
@@ -1,5 +1,6 @@
 // Copyright Epic Games, Inc. All Rights Reserved.
 #include "Net/Core/Trace/Private/NetTraceInternal.h"
+#include "AutoRTFM/AutoRTFM.h"
 #include "Containers/ArrayView.h"
 #include "Containers/Map.h"
 #include "Net/Core/Trace/NetDebugName.h"
@@ -54,7 +55,7 @@ struct FNetTraceInternal
 	{
 		FName Name;
 		const UE::Net::FNetDebugName* DebugName;
-		uint64 TypeIdentifier; 
+		uint64 TypeIdentifier;
 		uint32 OwnerId;
 	};
 
@@ -70,6 +71,8 @@ struct FNetTraceInternal
 	// Get next NameId used to track what we already have traced
 	static UE::Net::FNetDebugNameId GetNextNameId();
 
+	static uint32 GetCurrentTraceId();
+
 	static inline FThreadBuffer* GetThreadBuffer();
 
 	static constexpr ENetTraceVersion NetTraceVersion = ENetTraceVersion::ENetTraceVersion_DebugNameIndexIs32Bits;
@@ -88,21 +91,27 @@ FNetTrace::FOnNetTraceStarted FNetTrace::OnNetTraceStarted;
 FNetTraceInternal::FThreadBuffer* FNetTraceInternal::GetThreadBuffer()
 {
 	FNetTraceInternal::FThreadBuffer* ThreadBufferPtr = ThreadBuffer.Get();
-	if (!ThreadBufferPtr || ThreadBufferPtr->TraceId != CurrentTraceId)
+	uint32 TraceId = GetCurrentTraceId();
+	if (!ThreadBufferPtr || ThreadBufferPtr->TraceId != TraceId)
 	{
 		ThreadBuffer = MakeUnique<FNetTraceInternal::FThreadBuffer>();
 		ThreadBufferPtr = ThreadBuffer.Get();
-		ThreadBufferPtr->TraceId = CurrentTraceId;
+		ThreadBufferPtr->TraceId = TraceId;
 	}
-	
+
 	return ThreadBufferPtr;
 }
 
-UE::Net::FNetDebugNameId FNetTraceInternal::GetNextNameId()
+UE_AUTORTFM_ALWAYS_OPEN UE::Net::FNetDebugNameId FNetTraceInternal::GetNextNameId()
 {
 	return NextNameId++;
 }
 
+UE_AUTORTFM_ALWAYS_OPEN uint32 FNetTraceInternal::GetCurrentTraceId()
+{
+	return CurrentTraceId;
+}
+
 void FNetTrace::OnTraceStarted(FTraceAuxiliary::EConnectionType /*TraceType*/, const FString& /*TraceDestination*/)
 {
 	++CurrentTraceId;
@@ -163,7 +172,7 @@ void FNetTrace::SetTraceVerbosity(uint32 Verbosity)
 	else if (GetTraceVerbosity() && !NewVerbosity)
 	{
 		UE::Trace::ToggleChannel(TEXT("NetChannel"), false);
-			
+
 		if (UE::Trace::IsTracing())
 		{
 			OnTraceStopped(FTraceAuxiliary::EConnectionType::None, TEXT(""));
@@ -228,7 +237,7 @@ void FNetTrace::FoldTraceCollector(FNetTraceCollector* DstCollector, const FNetT
 
 		// When we fold non-bunch events we inject them at the current level
 		const uint32 Level = DstCollector->CurrentNestingLevel + LevelOffset;
-		
+
 		// Make sure that the events fit
 		if (SrcCollector->EventCount + DstCollector->EventCount > (uint32)DstCollector->Events.Num())
 		{
@@ -236,7 +245,7 @@ void FNetTrace::FoldTraceCollector(FNetTraceCollector* DstCollector, const FNetT
 		}
 
 		if (SrcCollector->EventCount + DstCollector->EventCount <= (uint32)DstCollector->Events.Num())
-		{			
+		{
 			uint32 DstEventIndex = DstCollector->EventCount;
 			FNetTracePacketContentEvent* DstEventData = DstCollector->Events.GetData();
 			for (const FNetTracePacketContentEvent& SrcEvent : MakeArrayView(SrcCollector->Events.GetData(), SrcCollector->EventCount))
@@ -249,11 +258,11 @@ void FNetTrace::FoldTraceCollector(FNetTraceCollector* DstCollector, const FNetT
 				DstEvent.NestingLevel += Level;
 				++DstEventIndex;
 			}
-			
+
 			DstCollector->EventCount += SrcCollector->EventCount;
 
 			DstCollector->ReferencedObjectIds.Append(SrcCollector->ReferencedObjectIds);
-		}		
+		}
 	}
 }
 
@@ -264,7 +273,7 @@ void FNetTrace::PushStreamOffset(FNetTraceCollector* Collector, uint32 Offset)
 		const uint32 OffsetStackLevel = Collector->OffsetStackLevel;
 
 		// Offset is additive, we always maintain a zero in the first slot
-		Collector->OffsetStack[OffsetStackLevel + 1] = Collector->OffsetStack[OffsetStackLevel] + Offset; 
+		Collector->OffsetStack[OffsetStackLevel + 1] = Collector->OffsetStack[OffsetStackLevel] + Offset;
 		++Collector->OffsetStackLevel;
 	}
 }
@@ -290,7 +299,7 @@ uint32 FNetTrace::BeginPacketContentEvent(FNetTraceCollector& Collector, ENetTra
 	{
 		Collector.Events.SetNumUninitialized(EventIndex + 1U, EAllowShrinking::No);
 	}
-	
+
 	FNetTracePacketContentEvent& Event = Collector.Events.GetData()[EventIndex];
 
 	// Add offset
@@ -317,7 +326,7 @@ void FNetTrace::EndPacketContentEvent(FNetTraceCollector& Collector, uint32 Even
 #endif
 
 	FNetTracePacketContentEvent& Event = Collector.Events.GetData()[EventIndex];
-	
+
 	const uint32 EventNestingLevel = Event.NestingLevel;
 	check(EventNestingLevel < Collector.CurrentNestingLevel);
 
@@ -338,8 +347,8 @@ void FNetTrace::TracePacketContentEvent(FNetTraceCollector& Collector, UE::Net::
 {
 	if (FNetTrace::GetTraceVerbosity() >= Verbosity)
 	{
-		const uint32 EventIndex = BeginPacketContentEvent(Collector, ENetTracePacketContentEventType::NameId, StartPos);	
-		
+		const uint32 EventIndex = BeginPacketContentEvent(Collector, ENetTracePacketContentEventType::NameId, StartPos);
+
 		FNetTracePacketContentEvent& Event = Collector.Events.GetData()[EventIndex];
 		Event.DebugNameId = InNetTraceNameId;
 
@@ -350,7 +359,7 @@ void FNetTrace::TracePacketContentEvent(FNetTraceCollector& Collector, UE::Net::
 void FNetTrace::BeginBunch(FNetTraceCollector& Collector)
 {
 #if UE_NET_TRACE_VALIDATE
-	check(Collector.PendingBunchEventIndex == ~0U);	
+	check(Collector.PendingBunchEventIndex == ~0U);
 	check(Collector.CurrentNestingLevel == 0U);
 #endif
 
@@ -433,7 +442,7 @@ void FNetTrace::TraceBunch(FNetTraceCollector& DstCollector, const FNetTraceBunc
 		FNetTrace::BeginBunch(DstCollector);
 		FNetTrace::FoldTraceCollector(&DstCollector, BunchCollector, 0U);
 	}
-		
+
 	FNetTrace::EndBunch(DstCollector, TraceName(BunchName), StartPos, HeaderBits, BunchBits, BunchInfo);
 }
 
@@ -444,7 +453,7 @@ void FNetTrace::TraceBunch(FNetTraceCollector& DstCollector, const FNetTraceBunc
 		FNetTrace::BeginBunch(DstCollector);
 		FNetTrace::FoldTraceCollector(&DstCollector, BunchCollector, 0U);
 	}
-		
+
 	FNetTrace::EndBunch(DstCollector, TraceName(BunchName), StartPos, HeaderBits, BunchBits, BunchInfo);
 }
 
@@ -479,7 +488,7 @@ void FNetTrace::TraceCollectedEvents(FNetTraceCollector& Collector, uint32 GameI
 	PacketInfo.PacketType = PacketType;
 
 	FNetTraceInternal::FGameInstance& GameInstance = GameInstances.FindOrAdd(GameInstanceId);
-	for(uint64 NetObjectId : Collector.ReferencedObjectIds)
+	for (uint64 NetObjectId : Collector.ReferencedObjectIds)
 	{
 		if (!GameInstance.ObjectsTraced.Contains(NetObjectId))
 		{
@@ -543,7 +552,7 @@ void FNetTrace::TracePacket(uint32 GameInstanceId, uint32 ConnectionId, uint32 P
 void FNetTrace::TraceObjectCreated(uint32 GameInstanceId, uint64 NetObjectId, const UE::Net::FNetDebugName* DebugName, uint64 TypeIdentifier, uint32 OwnerId)
 {
 	GameInstances.FindOrAdd(GameInstanceId).Objects.Emplace(NetObjectId, { NAME_None, DebugName, TypeIdentifier, OwnerId });
-	
+
 	if (!GNetTraceRuntimeVerbosity)
 	{
 		return;
@@ -560,8 +569,8 @@ void FNetTrace::TraceObjectCreated(uint32 GameInstanceId, uint64 NetObjectId, co
 
 void FNetTrace::TraceObjectCreated(uint32 GameInstanceId, uint64 NetObjectId, const FName ObjectName, uint64 TypeIdentifier, uint32 OwnerId)
 {
-	GameInstances.FindOrAdd(GameInstanceId).Objects.Emplace(NetObjectId, {ObjectName, nullptr, TypeIdentifier, OwnerId});
-	
+	GameInstances.FindOrAdd(GameInstanceId).Objects.Emplace(NetObjectId, { ObjectName, nullptr, TypeIdentifier, OwnerId });
+
 	if (!GNetTraceRuntimeVerbosity)
 	{
 		return;
@@ -679,8 +688,8 @@ UE::Net::FNetDebugNameId FNetTrace::TraceName(const TCHAR* Name)
 		ThreadBufferPtr->DynamicNameHashToNameIdMap.Add(HashedName, NameId);
 
 		FTCHARToUTF8 Converter(Name);
-		FNetTraceInternal::Reporter::ReportAnsiName(NameId, Converter.Length() + 1, (const char*)Converter.Get());		
-		
+		FNetTraceInternal::Reporter::ReportAnsiName(NameId, Converter.Length() + 1, (const char*)Converter.Get());
+
 		return NameId;
 	}
 }
@@ -709,8 +718,8 @@ UE::Net::FNetDebugNameId FNetTrace::TraceName(FName Name)
 		TCHAR Buffer[FName::StringBufferSize];
 		uint32 NameLen = Name.ToString(Buffer);
 		FTCHARToUTF8 Converter(Buffer);
-		FNetTraceInternal::Reporter::ReportAnsiName(NameId, Converter.Length() + 1, (const char*)Converter.Get());		
-		
+		FNetTraceInternal::Reporter::ReportAnsiName(NameId, Converter.Length() + 1, (const char*)Converter.Get());
+
 		return NameId;
 	}
 }
@@ -748,7 +757,7 @@ static FAutoConsoleCommand NeTraceSetVerbosityCmd = FAutoConsoleCommand(
 
 			FNetTrace::SetTraceVerbosity(FCString::Atoi(*Args[0]));
 		}
-		)
-	);
+	)
+);
 
 #endif
