Tracking Iris object poll/quantize/write in net traces

diff --git a/Engine/Source/Runtime/Experimental/Iris/Core/Private/Iris/ReplicationSystem/NetRefHandleManager.cpp b/Engine/Source/Runtime/Experimental/Iris/Core/Private/Iris/ReplicationSystem/NetRefHandleManager.cpp
--- a/Engine/Source/Runtime/Experimental/Iris/Core/Private/Iris/ReplicationSystem/NetRefHandleManager.cpp
+++ b/Engine/Source/Runtime/Experimental/Iris/Core/Private/Iris/ReplicationSystem/NetRefHandleManager.cpp
@@ -598,6 +598,8 @@ bool FNetRefHandleManager::AddSubObject(FNetRefHandle OwnerHandle, FNetRefHandle
 		return false;
 	}
 
+	UE_NET_TRACE_SUBOBJECT(OwnerHandle, SubObjectHandle);
+
 	const FInternalNetRefIndex RelativeOtherSubObjectInternalIndex = EnumHasAnyFlags(Flags, EAddSubObjectFlags::ReplicateWithSubObject) ? GetInternalIndex(RelativeOtherSubObjectHandle) : InvalidInternalIndex;
 	return InternalAddSubObject(OwnerInternalIndex, SubObjectInternalIndex, RelativeOtherSubObjectInternalIndex, Flags);
 }
diff --git a/Engine/Source/Runtime/Experimental/Iris/Core/Private/Iris/ReplicationSystem/Polling/ObjectPoller.cpp b/Engine/Source/Runtime/Experimental/Iris/Core/Private/Iris/ReplicationSystem/Polling/ObjectPoller.cpp
--- a/Engine/Source/Runtime/Experimental/Iris/Core/Private/Iris/ReplicationSystem/Polling/ObjectPoller.cpp
+++ b/Engine/Source/Runtime/Experimental/Iris/Core/Private/Iris/ReplicationSystem/Polling/ObjectPoller.cpp
@@ -13,6 +13,7 @@
 #include "Iris/Core/IrisProfiler.h"
 #include "Iris/Stats/NetStatsContext.h"
 
+#include "Net/Core/Trace/NetTrace.h"
 #include "Net/Core/Trace/NetDebugName.h"
 #include "HAL/IConsoleManager.h"
 
@@ -82,6 +83,9 @@ void FObjectPoller::ForcePollObject(FInternalNetRefIndex ObjectIndex)
 
 	IRIS_PROFILER_PROTOCOL_NAME(ObjectData.Protocol->DebugName->Name);
 
+	UE_NET_IRIS_STATS_TIMER(Timer, NetStatsContext);
+	UE_NET_TRACE_POLL_OBJECT_SCOPE(ObjectData.RefHandle, Timer);
+
 	// We always poll all states here.
 	ObjectData.bWantsFullPoll = 0U;
 
@@ -89,7 +93,6 @@ void FObjectPoller::ForcePollObject(FInternalNetRefIndex ObjectIndex)
 	if (EnumHasAnyFlags(ObjectData.InstanceProtocol->InstanceTraits, EReplicationInstanceProtocolTraits::NeedsPoll))
 	{
 		IRIS_PROFILER_SCOPE_VERBOSE(PollAndCopy);
-		UE_NET_IRIS_STATS_TIMER(Timer, NetStatsContext);
 
 		const bool bIsGCAffectedObject = GarbageCollectionAffectedObjects.GetBit(ObjectIndex);
 		GarbageCollectionAffectedObjects.ClearBit(ObjectIndex);
@@ -110,6 +113,7 @@ void FObjectPoller::ForcePollObject(FInternalNetRefIndex ObjectIndex)
 		else
 		{
 			UE_NET_IRIS_STATS_ADD_TIME_AND_COUNT_FOR_OBJECT_AS_WASTE(Timer, Poll, ObjectIndex);
+			UE_NET_TRACE_POLL_OBJECT_IS_WASTE();
 		}
 		++PollStats.PolledObjectCount;
 	}
@@ -130,6 +134,9 @@ void FObjectPoller::PushModelPollObject(FInternalNetRefIndex ObjectIndex)
 
 	IRIS_PROFILER_PROTOCOL_NAME(ObjectData.Protocol->DebugName->Name);
 
+	UE_NET_IRIS_STATS_TIMER(Timer, NetStatsContext);
+	UE_NET_TRACE_POLL_OBJECT_SCOPE(ObjectData.RefHandle, Timer);
+
 	const FReplicationInstanceProtocol* InstanceProtocol = ObjectData.InstanceProtocol;
 
 	const EReplicationInstanceProtocolTraits InstanceTraits = InstanceProtocol->InstanceTraits;
@@ -153,7 +160,6 @@ void FObjectPoller::PushModelPollObject(FInternalNetRefIndex ObjectIndex)
 	}
 
 	IRIS_PROFILER_SCOPE_VERBOSE(PollPushBased);
-	UE_NET_IRIS_STATS_TIMER(Timer, NetStatsContext);
 
 	// Does the object need to poll all states once.
 	const bool bWantsFullPoll = ObjectData.bWantsFullPoll;
@@ -213,6 +219,7 @@ void FObjectPoller::PushModelPollObject(FInternalNetRefIndex ObjectIndex)
 	else
 	{
 		UE_NET_IRIS_STATS_ADD_TIME_AND_COUNT_FOR_OBJECT_AS_WASTE(Timer, Poll, ObjectIndex);
+		UE_NET_TRACE_POLL_OBJECT_IS_WASTE();
 	}
 }
 
diff --git a/Engine/Source/Runtime/Experimental/Iris/Core/Private/Iris/ReplicationSystem/ReplicationOperationsInternal.cpp b/Engine/Source/Runtime/Experimental/Iris/Core/Private/Iris/ReplicationSystem/ReplicationOperationsInternal.cpp
--- a/Engine/Source/Runtime/Experimental/Iris/Core/Private/Iris/ReplicationSystem/ReplicationOperationsInternal.cpp
+++ b/Engine/Source/Runtime/Experimental/Iris/Core/Private/Iris/ReplicationSystem/ReplicationOperationsInternal.cpp
@@ -80,6 +80,7 @@ uint32 FReplicationInstanceOperationsInternal::QuantizeObjectStateData(FNetBitSt
 		{
 			IRIS_PROFILER_PROTOCOL_NAME(Object.Protocol->DebugName->Name);
 			UE_NET_IRIS_STATS_TIMER(Timer, SerializationContext.GetNetStatsContext());
+			UE_NET_TRACE_QUANTIZE_OBJECT_SCOPE(Object.RefHandle, Timer);
 
 			// if the object was scopable prev frame we can do partial copy
 			bool bShouldPropagateChangedStates = Object.bShouldPropagateChangedStates;
diff --git a/Engine/Source/Runtime/Experimental/Iris/Core/Private/Iris/ReplicationSystem/ReplicationWriter.cpp b/Engine/Source/Runtime/Experimental/Iris/Core/Private/Iris/ReplicationSystem/ReplicationWriter.cpp
--- a/Engine/Source/Runtime/Experimental/Iris/Core/Private/Iris/ReplicationSystem/ReplicationWriter.cpp
+++ b/Engine/Source/Runtime/Experimental/Iris/Core/Private/Iris/ReplicationSystem/ReplicationWriter.cpp
@@ -2801,24 +2801,27 @@ FReplicationWriter::EWriteObjectStatus FReplicationWriter::WriteObjectAndSubObje
 
 FReplicationWriter::EWriteObjectStatus FReplicationWriter::WriteObjectInBatch(FNetSerializationContext& Context, uint32 InternalIndex, uint32 WriteObjectFlags, FBatchInfo& OutBatchInfo)
 {
-	UE_NET_IRIS_STATS_TIMER(Timer, Context.GetNetStatsContext());
-
-	// Reset pending exports
-	FNetExportContext* ExportContext = Context.GetExportContext();
-	if (ExportContext)
 	{
-		ExportContext->ClearPendingExports();
-	}
+		UE_NET_IRIS_STATS_TIMER(Timer, Context.GetNetStatsContext());
+		UE_NET_TRACE_WRITE_OBJECT_SCOPE(NetRefHandleManager->GetReplicatedObjectDataNoCheck(InternalIndex).RefHandle, Timer);
 
-	// Write parent object and subobjects
-	const EWriteObjectStatus WriteObjectStatus = WriteObjectAndSubObjects(Context, InternalIndex, WriteObjectFlags, OutBatchInfo);
-	if (!IsWriteObjectSuccess(WriteObjectStatus))
-	{
-		UE_NET_IRIS_STATS_ADD_TIME_AND_COUNT_FOR_OBJECT_AS_WASTE(Timer, Write, InternalIndex);
-		return WriteObjectStatus;
-	}
+		// Reset pending exports
+		FNetExportContext* ExportContext = Context.GetExportContext();
+		if (ExportContext)
+		{
+			ExportContext->ClearPendingExports();
+		}
 
-	UE_NET_IRIS_STATS_ADD_TIME_AND_COUNT_FOR_OBJECT(Timer, Write, InternalIndex);
+		// Write parent object and subobjects
+		const EWriteObjectStatus WriteObjectStatus = WriteObjectAndSubObjects(Context, InternalIndex, WriteObjectFlags, OutBatchInfo);
+		if (!IsWriteObjectSuccess(WriteObjectStatus))
+		{
+			UE_NET_IRIS_STATS_ADD_TIME_AND_COUNT_FOR_OBJECT_AS_WASTE(Timer, Write, InternalIndex);
+			return WriteObjectStatus;
+		}
+
+		UE_NET_IRIS_STATS_ADD_TIME_AND_COUNT_FOR_OBJECT(Timer, Write, InternalIndex);
+	}
 
 	// Include dependent objects as separate batch, (for hugeobjects they will be included as they are written to a separate bitstream)
 	{
diff --git a/Engine/Source/Runtime/Experimental/Iris/Core/Private/Iris/Stats/NetStatsContext.h b/Engine/Source/Runtime/Experimental/Iris/Core/Private/Iris/Stats/NetStatsContext.h
--- a/Engine/Source/Runtime/Experimental/Iris/Core/Private/Iris/Stats/NetStatsContext.h
+++ b/Engine/Source/Runtime/Experimental/Iris/Core/Private/Iris/Stats/NetStatsContext.h
@@ -7,6 +7,7 @@
 #include "ProfilingDebugging/CsvProfilerConfig.h"
 #include "Iris/ReplicationSystem/NetRefHandleManager.h"
 #include "Iris/ReplicationSystem/ReplicationProtocol.h"
+#include "Net/Core/Trace/NetTrace.h"
 
 namespace UE::Net::Private
 {
@@ -82,23 +83,20 @@ private:
 	bool bShouldIncludeSubObjectWithRoot = false;
 };
 
-class FNetStatsTimer
+class FNetStatsTimer : public FNetTimer
 {
 public:
 	FNetStatsTimer(FNetStatsContext* InStatsContext)
-	: NetStatsContext(InStatsContext)
-	, StartCycle(InStatsContext ? FPlatformTime::Cycles64() : uint64(0)) 
+	: FNetTimer(InStatsContext != nullptr)
+	, NetStatsContext(InStatsContext)
 	{
 	}
 
-	uint64 GetCyclesSinceStart() const { return FPlatformTime::Cycles64() - StartCycle; }
-
 	// If we have a context, this timer is enabled 
 	FNetStatsContext* GetNetStatsContext() const { return NetStatsContext; }
 	
 private:
 	FNetStatsContext* NetStatsContext;
-	uint64 StartCycle;
 };
 
 inline FNetTypeStatsData& FNetStatsContext::GetTypeStatsData(FNetStatsContext& Context, const FNetRefHandleManager::FReplicatedObjectData& ObjectData, bool bTreatAsRoot)
@@ -148,8 +146,6 @@ inline void FNetStatsContext::ResetStats(int32 NumTypeStats)
 // Wrap usage of stats in macros so we can compile it out
 #if UE_NET_IRIS_CSV_STATS
 
-#define UE_NET_IRIS_STATS_TIMER(TimerName, NetStatsContext) UE::Net::Private::FNetStatsTimer TimerName(NetStatsContext);
-
 #define UE_NET_IRIS_STATS_ADD_TIME_AND_COUNT_FOR_OBJECT(Timer, StatName, ObjectIndex) \
 	do { \
 		if (UE::Net::Private::FNetStatsContext* LocalNetStatsContext = Timer.GetNetStatsContext()) \
@@ -227,7 +223,6 @@ inline void FNetStatsContext::ResetStats(int32 NumTypeStats)
 
 #else
 
-#define UE_NET_IRIS_STATS_TIMER(...)
 #define UE_NET_IRIS_STATS_ADD_TIME_AND_COUNT_FOR_OBJECT(...)
 #define UE_NET_IRIS_STATS_ADD_TIME_AND_COUNT_FOR_OBJECT_AS_WASTE(...)
 #define UE_NET_IRIS_STATS_ADD_BITS_WRITTEN_AND_COUNT_FOR_OBJECT(...)
@@ -237,3 +232,14 @@ inline void FNetStatsContext::ResetStats(int32 NumTypeStats)
 #define UE_NET_IRIS_STATS_ADD_COUNT_FOR_OBJECT(...)
 
 #endif
+
+// this is needed for net traces as well
+#if UE_NET_IRIS_CSV_STATS || UE_NET_TRACE_ENABLED
+
+#define UE_NET_IRIS_STATS_TIMER(TimerName, NetStatsContext) UE::Net::Private::FNetStatsTimer TimerName(NetStatsContext);
+
+#else
+
+#define UE_NET_IRIS_STATS_TIMER(...)
+
+#endif
diff --git a/Engine/Source/Runtime/Net/Core/Private/Net/Core/Trace/NetTraceInternal.cpp b/Engine/Source/Runtime/Net/Core/Private/Net/Core/Trace/NetTraceInternal.cpp
--- a/Engine/Source/Runtime/Net/Core/Private/Net/Core/Trace/NetTraceInternal.cpp
+++ b/Engine/Source/Runtime/Net/Core/Private/Net/Core/Trace/NetTraceInternal.cpp
@@ -55,7 +55,8 @@ struct FNetTraceInternal
 	{
 		FName Name;
 		const UE::Net::FNetDebugName* DebugName;
-		uint64 TypeIdentifier;
+		uint64 TypeIdentifier; 
+		uint64 ParentNetObjectId;
 		uint32 OwnerId;
 	};
 
@@ -479,6 +480,38 @@ void FNetTrace::PopSendBunch(FNetTraceCollector& Collector)
 	BunchHeaderEvent.EndPos = 0U;
 }
 
+static void TraceExistingObject(uint32 GameInstanceId, const TMap<uint64, FNetTraceInternal::FNetObject>& Objects, TSet<uint64>& ObjectsTraced, uint64 NetObjectId)
+{
+	if (!ObjectsTraced.Contains(NetObjectId))
+	{
+		const FNetTraceInternal::FNetObject* NetObject = Objects.Find(NetObjectId);
+		if (NetObject)
+		{
+			ObjectsTraced.Add(NetObjectId);
+
+			if (NetObject->DebugName)
+			{
+				if (NetObject->DebugName->DebugNameId == 0U)
+				{
+					NetObject->DebugName->DebugNameId = FNetTrace::TraceName(NetObject->DebugName->Name);
+				}
+					
+				FNetTraceInternal::Reporter::ReportObjectExists(GameInstanceId, NetObjectId, NetObject->DebugName->DebugNameId, NetObject->TypeIdentifier, NetObject->OwnerId);
+			}
+			else
+			{
+				FNetTraceInternal::Reporter::ReportObjectExists(GameInstanceId, NetObjectId, FNetTrace::TraceName(NetObject->Name), NetObject->TypeIdentifier, NetObject->OwnerId);
+			}
+
+			if (NetObject->ParentNetObjectId != uint64(-1))
+			{
+				TraceExistingObject(GameInstanceId, Objects, ObjectsTraced, NetObject->ParentNetObjectId);
+				FNetTraceInternal::Reporter::ReportSubObject(GameInstanceId, NetObject->ParentNetObjectId, NetObjectId);
+			}
+		}
+	}
+}
+
 void FNetTrace::TraceCollectedEvents(FNetTraceCollector& Collector, uint32 GameInstanceId, uint32 ConnectionId, ENetTracePacketType PacketType)
 {
 	FNetTracePacketInfo PacketInfo;
@@ -490,27 +523,7 @@ void FNetTrace::TraceCollectedEvents(FNetTraceCollector& Collector, uint32 GameI
 	FNetTraceInternal::FGameInstance& GameInstance = GameInstances.FindOrAdd(GameInstanceId);
 	for (uint64 NetObjectId : Collector.ReferencedObjectIds)
 	{
-		if (!GameInstance.ObjectsTraced.Contains(NetObjectId))
-		{
-			const FNetTraceInternal::FNetObject* NetObject = GameInstance.Objects.Find(NetObjectId);
-			if (NetObject)
-			{
-				GameInstance.ObjectsTraced.Add(NetObjectId);
-
-				if (NetObject->DebugName)
-				{
-					if (NetObject->DebugName->DebugNameId == 0U)
-					{
-						NetObject->DebugName->DebugNameId = FNetTrace::TraceName(NetObject->DebugName->Name);
-					}
-					FNetTraceInternal::Reporter::ReportObjectExists(GameInstanceId, NetObjectId, NetObject->DebugName->DebugNameId, NetObject->TypeIdentifier, NetObject->OwnerId);
-				}
-				else
-				{
-					FNetTraceInternal::Reporter::ReportObjectExists(GameInstanceId, NetObjectId, FNetTrace::TraceName(NetObject->Name), NetObject->TypeIdentifier, NetObject->OwnerId);
-				}
-			}
-		}
+		TraceExistingObject(GameInstanceId, GameInstance.Objects, GameInstance.ObjectsTraced, NetObjectId);
 	}
 
 	// Trace all collected events
@@ -551,7 +564,7 @@ void FNetTrace::TracePacket(uint32 GameInstanceId, uint32 ConnectionId, uint32 P
 
 void FNetTrace::TraceObjectCreated(uint32 GameInstanceId, uint64 NetObjectId, const UE::Net::FNetDebugName* DebugName, uint64 TypeIdentifier, uint32 OwnerId)
 {
-	GameInstances.FindOrAdd(GameInstanceId).Objects.Emplace(NetObjectId, { NAME_None, DebugName, TypeIdentifier, OwnerId });
+	GameInstances.FindOrAdd(GameInstanceId).Objects.Emplace(NetObjectId, { .Name = NAME_None, .DebugName = DebugName, .TypeIdentifier = TypeIdentifier, .ParentNetObjectId = uint64(-1), .OwnerId = OwnerId });
 
 	if (!GNetTraceRuntimeVerbosity)
 	{
@@ -569,8 +582,8 @@ void FNetTrace::TraceObjectCreated(uint32 GameInstanceId, uint64 NetObjectId, co
 
 void FNetTrace::TraceObjectCreated(uint32 GameInstanceId, uint64 NetObjectId, const FName ObjectName, uint64 TypeIdentifier, uint32 OwnerId)
 {
-	GameInstances.FindOrAdd(GameInstanceId).Objects.Emplace(NetObjectId, { ObjectName, nullptr, TypeIdentifier, OwnerId });
-
+	GameInstances.FindOrAdd(GameInstanceId).Objects.Emplace(NetObjectId, { .Name = ObjectName, .DebugName = nullptr, .TypeIdentifier = TypeIdentifier, .ParentNetObjectId = uint64(-1), .OwnerId = OwnerId });
+	
 	if (!GNetTraceRuntimeVerbosity)
 	{
 		return;
@@ -744,6 +757,49 @@ UE::Net::FNetDebugNameId FNetTrace::TraceName(const UE::Net::FNetDebugName* Debu
 	}
 }
 
+void FNetTrace::TraceObjectPolled(uint32 GameInstanceId, uint64 NetObjectId, uint64 Cycles, bool bIsWaste)
+{
+	// not rechecking GNetTraceRuntimeVerbosity here as it's already checked to be == VeryVerbose to have got this far
+
+	FNetTraceInternal::FGameInstance& GameInstance = GameInstances.FindOrAdd(GameInstanceId);
+	TraceExistingObject(GameInstanceId, GameInstance.Objects, GameInstance.ObjectsTraced, NetObjectId);
+
+	FNetTraceInternal::Reporter::ReportObjectPolled(GameInstanceId, NetObjectId, Cycles, bIsWaste);
+}
+
+void FNetTrace::TraceObjectQuantized(uint32 GameInstanceId, uint64 NetObjectId, uint64 Cycles)
+{
+	// not rechecking GNetTraceRuntimeVerbosity here as it's already checked to be == VeryVerbose to have got this far
+
+	FNetTraceInternal::FGameInstance& GameInstance = GameInstances.FindOrAdd(GameInstanceId);
+	TraceExistingObject(GameInstanceId, GameInstance.Objects, GameInstance.ObjectsTraced, NetObjectId);
+
+	FNetTraceInternal::Reporter::ReportObjectQuantized(GameInstanceId, NetObjectId, Cycles);
+}
+
+void FNetTrace::TraceObjectWritten(uint32 GameInstanceId, uint64 NetObjectId, uint64 Cycles)
+{
+	// not rechecking GNetTraceRuntimeVerbosity here as it's already checked to be == VeryVerbose to have got this far
+
+	FNetTraceInternal::FGameInstance& GameInstance = GameInstances.FindOrAdd(GameInstanceId);
+	TraceExistingObject(GameInstanceId, GameInstance.Objects, GameInstance.ObjectsTraced, NetObjectId);
+
+	FNetTraceInternal::Reporter::ReportObjectWritten(GameInstanceId, NetObjectId, Cycles);
+}
+
+void FNetTrace::TraceSubObject(uint32 GameInstanceId, uint64 ParentNetObjectId, uint64 SubObjectNetObjectId)
+{
+	if (FNetTraceInternal::FNetObject* SubObject = GameInstances.FindOrAdd(GameInstanceId).Objects.Find(SubObjectNetObjectId))
+	{
+		SubObject->ParentNetObjectId = ParentNetObjectId;
+	}
+	
+	if (GNetTraceRuntimeVerbosity)
+	{
+		FNetTraceInternal::Reporter::ReportSubObject(GameInstanceId, ParentNetObjectId, SubObjectNetObjectId);
+	}	
+}
+
 static FAutoConsoleCommand NeTraceSetVerbosityCmd = FAutoConsoleCommand(
 	TEXT("NetTrace.SetTraceVerbosity"),
 	TEXT("Start NetTrace with given verbositylevel."),
diff --git a/Engine/Source/Runtime/Net/Core/Private/Net/Core/Trace/Reporters/NetTraceReporter.cpp b/Engine/Source/Runtime/Net/Core/Private/Net/Core/Trace/Reporters/NetTraceReporter.cpp
--- a/Engine/Source/Runtime/Net/Core/Private/Net/Core/Trace/Reporters/NetTraceReporter.cpp
+++ b/Engine/Source/Runtime/Net/Core/Private/Net/Core/Trace/Reporters/NetTraceReporter.cpp
@@ -129,6 +129,31 @@ UE_TRACE_EVENT_BEGIN(NetTrace, PacketDroppedEvent)
 	UE_TRACE_EVENT_FIELD(uint8, PacketType)
 UE_TRACE_EVENT_END()
 
+UE_TRACE_EVENT_BEGIN(NetTrace, ObjectPolledEvent)
+	UE_TRACE_EVENT_FIELD(uint64, ObjectId)
+	UE_TRACE_EVENT_FIELD(uint64, Cycles)
+	UE_TRACE_EVENT_FIELD(bool, bIsWaste)
+	UE_TRACE_EVENT_FIELD(uint8, GameInstanceId)
+UE_TRACE_EVENT_END()
+
+UE_TRACE_EVENT_BEGIN(NetTrace, ObjectQuantizedEvent)
+	UE_TRACE_EVENT_FIELD(uint64, ObjectId)
+	UE_TRACE_EVENT_FIELD(uint64, Cycles)
+	UE_TRACE_EVENT_FIELD(uint8, GameInstanceId)
+UE_TRACE_EVENT_END()
+
+UE_TRACE_EVENT_BEGIN(NetTrace, ObjectWrittenEvent)
+	UE_TRACE_EVENT_FIELD(uint64, ObjectId)
+	UE_TRACE_EVENT_FIELD(uint64, Cycles)
+	UE_TRACE_EVENT_FIELD(uint8, GameInstanceId)
+UE_TRACE_EVENT_END()
+
+UE_TRACE_EVENT_BEGIN(NetTrace, SubObjectEvent)
+	UE_TRACE_EVENT_FIELD(uint64, ParentObjectId)
+	UE_TRACE_EVENT_FIELD(uint64, SubObjectId)
+	UE_TRACE_EVENT_FIELD(uint8, GameInstanceId)
+UE_TRACE_EVENT_END()
+
 void FNetTraceReporter::ReportInitEvent(uint32 NetTraceVersion)
 {
 	UE_TRACE_LOG(NetTrace, InitEvent, NetChannel)
@@ -370,4 +395,37 @@ void FNetTraceReporter::ReportObjectDestroyed(uint32 GameInstanceId, uint64 NetO
 		<< ObjectDestroyedEvent.GameInstanceId((uint8)GameInstanceId);
 }
 
+void FNetTraceReporter::ReportObjectPolled(uint32 GameInstanceId, uint64 NetObjectId, uint64 Cycles, bool bIsWaste)
+{
+	UE_TRACE_LOG(NetTrace, ObjectPolledEvent, NetChannel)
+		<< ObjectPolledEvent.ObjectId(NetObjectId)
+		<< ObjectPolledEvent.Cycles(Cycles)
+		<< ObjectPolledEvent.bIsWaste(bIsWaste)
+		<< ObjectPolledEvent.GameInstanceId((uint8)GameInstanceId);
+}
+
+void FNetTraceReporter::ReportObjectQuantized(uint32 GameInstanceId, uint64 NetObjectId, uint64 Cycles)
+{
+	UE_TRACE_LOG(NetTrace, ObjectQuantizedEvent, NetChannel)
+		<< ObjectQuantizedEvent.ObjectId(NetObjectId)
+		<< ObjectQuantizedEvent.Cycles(Cycles)
+		<< ObjectQuantizedEvent.GameInstanceId((uint8)GameInstanceId);
+}
+
+void FNetTraceReporter::ReportObjectWritten(uint32 GameInstanceId, uint64 NetObjectId, uint64 Cycles)
+{
+	UE_TRACE_LOG(NetTrace, ObjectWrittenEvent, NetChannel)
+		<< ObjectWrittenEvent.ObjectId(NetObjectId)
+		<< ObjectWrittenEvent.Cycles(Cycles)
+		<< ObjectWrittenEvent.GameInstanceId((uint8)GameInstanceId);
+}
+
+void FNetTraceReporter::ReportSubObject(uint32 GameInstanceId, uint64 ParentObjectId, uint64 SubObjectId)
+{
+	UE_TRACE_LOG(NetTrace, SubObjectEvent, NetChannel)
+		<< SubObjectEvent.ParentObjectId(ParentObjectId)
+		<< SubObjectEvent.SubObjectId(SubObjectId)
+		<< SubObjectEvent.GameInstanceId((uint8)GameInstanceId);
+}
+
 #endif
diff --git a/Engine/Source/Runtime/Net/Core/Private/Net/Core/Trace/Reporters/NetTraceReporter.h b/Engine/Source/Runtime/Net/Core/Private/Net/Core/Trace/Reporters/NetTraceReporter.h
--- a/Engine/Source/Runtime/Net/Core/Private/Net/Core/Trace/Reporters/NetTraceReporter.h
+++ b/Engine/Source/Runtime/Net/Core/Private/Net/Core/Trace/Reporters/NetTraceReporter.h
@@ -26,6 +26,10 @@ struct FNetTraceReporter
 	static void ReportFrameStatsCounter(uint32 GameInstanceId, UE::Net::FNetDebugNameId CounterNameId, uint32 StatValue);
 	static void ReportInstanceUpdated(uint32 GameInstanceId, bool bIsServer, const TCHAR* Name);
 	static void ReportInstanceDestroyed(uint32 GameInstanceId);
+	static void ReportObjectPolled(uint32 GameInstanceId, uint64 NetObjectId, uint64 Cycles, bool bIsWaste);
+	static void ReportObjectQuantized(uint32 GameInstanceId, uint64 NetObjectId, uint64 Cycles);
+	static void ReportObjectWritten(uint32 GameInstanceId, uint64 NetObjectId, uint64 Cycles);
+	static void ReportSubObject(uint32 GameInstanceId, uint64 ParentObjectId, uint64 SubObjectId);
 };
 
 #endif
diff --git a/Engine/Source/Runtime/Net/Core/Public/Net/Core/Trace/NetTrace.h b/Engine/Source/Runtime/Net/Core/Public/Net/Core/Trace/NetTrace.h
--- a/Engine/Source/Runtime/Net/Core/Public/Net/Core/Trace/NetTrace.h
+++ b/Engine/Source/Runtime/Net/Core/Public/Net/Core/Trace/NetTrace.h
@@ -123,6 +123,13 @@
 /** Trace additional information about the given game instance */
 #define UE_NET_TRACE_UPDATE_INSTANCE(...) UE_NET_TRACE_INTERNAL_UPDATE_INSTANCE(__VA_ARGS__)
 
+#define UE_NET_TRACE_POLL_OBJECT_SCOPE(...) UE_NET_TRACE_INTERNAL_POLL_OBJECT_SCOPE(__VA_ARGS__)
+#define UE_NET_TRACE_POLL_OBJECT_IS_WASTE() UE_NET_TRACE_INTERNAL_POLL_OBJECT_IS_WASTE
+#define UE_NET_TRACE_QUANTIZE_OBJECT_SCOPE(...) UE_NET_TRACE_INTERNAL_QUANTIZE_OBJECT_SCOPE(__VA_ARGS__)
+#define UE_NET_TRACE_WRITE_OBJECT_SCOPE(...) UE_NET_TRACE_INTERNAL_WRITE_OBJECT_SCOPE(__VA_ARGS__)
+
+#define UE_NET_TRACE_SUBOBJECT(ParentHandle, SubObjectHandle) UE_NET_TRACE_INTERNAL_SUBOBJECT(ParentHandle, SubObjectHandle)
+
 #else
 
 #define UE_NET_TRACE_SCOPE(...)
@@ -168,4 +175,11 @@
 
 #define UE_NET_TRACE_UPDATE_INSTANCE(...)
 
+#define UE_NET_TRACE_POLL_OBJECT_SCOPE(...)
+#define UE_NET_TRACE_POLL_OBJECT_IS_WASTE()
+#define UE_NET_TRACE_QUANTIZE_OBJECT_SCOPE(...)
+#define UE_NET_TRACE_WRITE_OBJECT_SCOPE(...)
+
+#define UE_NET_TRACE_SUBOBJECT(...)
+
 #endif // UE_NET_TRACE_ENABLED
diff --git a/Engine/Source/Runtime/Net/Core/Public/Net/Core/Trace/Private/NetTraceInternal.h b/Engine/Source/Runtime/Net/Core/Public/Net/Core/Trace/Private/NetTraceInternal.h
--- a/Engine/Source/Runtime/Net/Core/Public/Net/Core/Trace/Private/NetTraceInternal.h
+++ b/Engine/Source/Runtime/Net/Core/Public/Net/Core/Trace/Private/NetTraceInternal.h
@@ -3,6 +3,25 @@
 #pragma once
 
 #include "Net/Core/Trace/NetTraceConfig.h"
+#include "HAL/PlatformTime.h"
+
+namespace UE::Net::Private
+{
+	class FNetTimer
+	{
+	public:
+		FNetTimer(bool bEnabled)
+		: StartCycle(bEnabled ? FPlatformTime::Cycles64() : uint64(0)) 
+		{
+		}
+
+		uint64 GetCyclesSinceStart() const { return FPlatformTime::Cycles64() - StartCycle; }
+	
+	private:
+		uint64 StartCycle;
+	};
+}
+
 
 #if UE_NET_TRACE_ENABLED
 
@@ -280,6 +299,14 @@ struct FNetTrace
 	/** Trace */
 	inline static UE::Net::FNetDebugNameId TraceName(const FString& Name) { return TraceName(ToCStr(Name)); };
 
+	NETCORE_API static void TraceObjectPolled(uint32 GameInstanceId, uint64 NetObjectId, uint64 Cycles, bool bIsDirty);
+
+	NETCORE_API static void TraceObjectQuantized(uint32 GameInstanceId, uint64 NetObjectId, uint64 Cycles);
+
+	NETCORE_API static void TraceObjectWritten(uint32 GameInstanceId, uint64 NetObjectId, uint64 Cycles);
+
+	NETCORE_API static void TraceSubObject(uint32 GameInstanceId, uint64 ParentNetObjectId, uint64 SubObjectNetObjectId);
+
 	static bool constexpr GetNetTraceVerbosityEnabled(const uint32 V) { return (uint32)UE_NET_TRACE_COMPILETIME_VERBOSITY >= V; }
 
 	inline static FNetTraceCollector* GetCollectorAtVerbosity(FNetTraceCollector* Collector, uint32 CollectorVerbosity);
@@ -402,6 +429,66 @@ private:
 	FNetTraceCollector* Collector;
 };
 
+struct FNetTraceObjectTimer
+{
+	FNetTraceObjectTimer(const UE::Net::Private::FNetTimer& InTimer, uint32 InGameInstanceId, uint64 InNetObjectId)
+	: Timer(InTimer)
+	, NetObjectId(InNetObjectId)
+	, GameInstanceId(InGameInstanceId)
+	{}
+
+	const UE::Net::Private::FNetTimer& Timer;
+	uint64 NetObjectId;
+	uint32 GameInstanceId;
+};
+
+struct FNetTraceObjectPollTimer : public FNetTraceObjectTimer
+{
+	FNetTraceObjectPollTimer(const UE::Net::Private::FNetTimer& InTimer, uint32 InGameInstanceId, uint64 InNetObjectId) 
+	: FNetTraceObjectTimer(InTimer, InGameInstanceId, InNetObjectId)
+	, bIsWaste(false) {}
+
+	~FNetTraceObjectPollTimer() 
+	{ 
+		if (GameInstanceId != -1)
+		{
+			FNetTrace::TraceObjectPolled(GameInstanceId, NetObjectId, Timer.GetCyclesSinceStart(), bIsWaste);
+		}
+	}
+	
+	bool bIsWaste;
+};
+
+struct FNetTraceObjectQuantizeTimer : public FNetTraceObjectTimer
+{
+	FNetTraceObjectQuantizeTimer(const UE::Net::Private::FNetTimer& InTimer, uint32 InGameInstanceId, uint64 InNetObjectId) 
+	: FNetTraceObjectTimer(InTimer, InGameInstanceId, InNetObjectId)
+	{}
+
+	~FNetTraceObjectQuantizeTimer() 
+	{ 
+		if (GameInstanceId != -1)
+		{
+			FNetTrace::TraceObjectQuantized(GameInstanceId, NetObjectId, Timer.GetCyclesSinceStart());
+		}
+	}
+};
+
+struct FNetTraceObjectWriteTimer : public FNetTraceObjectTimer
+{
+	FNetTraceObjectWriteTimer(const UE::Net::Private::FNetTimer& InTimer, uint32 InGameInstanceId, uint64 InNetObjectId) 
+	: FNetTraceObjectTimer(InTimer, InGameInstanceId, InNetObjectId)
+	{}
+
+	~FNetTraceObjectWriteTimer()
+	{
+		if (GameInstanceId != -1)
+		{
+			FNetTrace::TraceObjectWritten(GameInstanceId, NetObjectId, Timer.GetCyclesSinceStart());
+		}
+	}
+};
+
 // FNetTrace implementation
 FNetTraceCollector* FNetTrace::CreateTraceCollector(uint32 Verbosity)
 {
@@ -630,5 +717,12 @@ void FNetTraceCollector::Reset()
 
 #define UE_NET_TRACE_INTERNAL_END_SESSION(GameInstanceId) FNetTrace::TraceEndSession(GameInstanceId);
 #define UE_NET_TRACE_INTERNAL_UPDATE_INSTANCE(...) FNetTrace::TraceInstanceUpdated(__VA_ARGS__)
+
+	#define UE_NET_TRACE_INTERNAL_POLL_OBJECT_SCOPE(Handle, Timer) FNetTraceObjectPollTimer __NetTracePollTimer(Timer, (GNetTraceRuntimeVerbosity >= ENetTraceVerbosity::VeryVerbose && Handle.IsValid()) ? Handle.GetReplicationSystemId() : -1, Handle.GetId());
+#define UE_NET_TRACE_INTERNAL_POLL_OBJECT_IS_WASTE __NetTracePollTimer.bIsWaste = true;
+#define UE_NET_TRACE_INTERNAL_QUANTIZE_OBJECT_SCOPE(Handle, Timer) FNetTraceObjectQuantizeTimer PREPROCESSOR_JOIN(__NetTraceQuantizeTimer, __LINE__)(Timer, (GNetTraceRuntimeVerbosity >= ENetTraceVerbosity::VeryVerbose && Handle.IsValid()) ? Handle.GetReplicationSystemId() : -1, Handle.GetId());
+#define UE_NET_TRACE_INTERNAL_WRITE_OBJECT_SCOPE(Handle, Timer) FNetTraceObjectWriteTimer PREPROCESSOR_JOIN(__NetTraceWriteTimer, __LINE__)(Timer, (GNetTraceRuntimeVerbosity >= ENetTraceVerbosity::VeryVerbose && Handle.IsValid()) ? Handle.GetReplicationSystemId() : -1, Handle.GetId());	
+	
+#define UE_NET_TRACE_INTERNAL_SUBOBJECT(ParentHandle, SubObjectHandle) FNetTrace::TraceSubObject(ParentHandle.GetReplicationSystemId(), GetObjectIdForNetTrace(ParentHandle), GetObjectIdForNetTrace(SubObjectHandle))
 					
 #endif // UE_NET_TRACE_ENABLED
