Net traces: fix net stat frame counter names in partial traces

diff --git a/Engine/Source/Runtime/Net/Core/Public/Net/Core/Trace/Private/NetTraceInternal.h b/Engine/Source/Runtime/Net/Core/Public/Net/Core/Trace/Private/NetTraceInternal.h
--- a/Engine/Source/Runtime/Net/Core/Public/Net/Core/Trace/Private/NetTraceInternal.h
+++ b/Engine/Source/Runtime/Net/Core/Public/Net/Core/Trace/Private/NetTraceInternal.h
@@ -706,13 +706,13 @@ void FNetTraceCollector::Reset()
 #define UE_NET_TRACE_INTERNAL_PACKET_STATSCOUNTER(GameInstanceId, ConnectionId, Name, StatValue, Verbosity) \
 	do { if (FNetTrace::GetNetTraceVerbosityEnabled(Verbosity)) { \
 			static const UE::Net::FNetDebugName* PREPROCESSOR_JOIN(NetTraceNameId_, __LINE__) = UE::Net::CreatePersistentNetDebugName(TEXT(#Name)); \
-			FNetTrace::TracePacketStatsCounter(GameInstanceId, ConnectionId, PREPROCESSOR_JOIN(NetTraceNameId_, __LINE__), StatValue); \
+			FNetTrace::TracePacketStatsCounter(GameInstanceId, ConnectionId, FNetTrace::TraceName(PREPROCESSOR_JOIN(NetTraceNameId_, __LINE__)), StatValue); \
 	} } while (0)
 
 #define UE_NET_TRACE_INTERNAL_FRAME_STATSCOUNTER(GameInstanceId, Name, StatValue, Verbosity) \
 	do { if (FNetTrace::GetNetTraceVerbosityEnabled(Verbosity)) { \
 			static const UE::Net::FNetDebugName* PREPROCESSOR_JOIN(NetTraceNameId_, __LINE__) = UE::Net::CreatePersistentNetDebugName(TEXT(#Name)); \
-			FNetTrace::TraceFrameStatsCounter(GameInstanceId, PREPROCESSOR_JOIN(NetTraceNameId_, __LINE__), StatValue); \
+			FNetTrace::TraceFrameStatsCounter(GameInstanceId, FNetTrace::TraceName(PREPROCESSOR_JOIN(NetTraceNameId_, __LINE__)), StatValue); \
 	} } while (0)
 
 #define UE_NET_TRACE_INTERNAL_END_SESSION(GameInstanceId) FNetTrace::TraceEndSession(GameInstanceId);
