Removed hidden friend in NetResult because it is "disabled" anyway by adding a declaration outside the class.

diff --git a/Engine/Source/Runtime/Net/Core/Public/Net/Core/Connection/NetResult.h b/Engine/Source/Runtime/Net/Core/Public/Net/Core/Connection/NetResult.h
--- a/Engine/Source/Runtime/Net/Core/Public/Net/Core/Connection/NetResult.h
+++ b/Engine/Source/Runtime/Net/Core/Public/Net/Core/Connection/NetResult.h
@@ -286,7 +286,7 @@ private:
  * @param ResultPtr		The result pointer to chain or assign the new result to.
  * @param Result		The result to chain to the result pointer. Should use MoveTemp.
  */
-static void AddToChainResultPtr(TPimplPtr<FNetResult, EPimplPtrMode::DeepCopy>& ResultPtr, FNetResult&& Result)
+inline void AddToChainResultPtr(TPimplPtr<FNetResult, EPimplPtrMode::DeepCopy>& ResultPtr, FNetResult&& Result)
 {
 	if (ResultPtr.IsValid())
 	{
@@ -305,7 +305,7 @@ static void AddToChainResultPtr(TPimplPtr<FNetResult, EPimplPtrMode::DeepCopy>&
  * @param Result		The result to chain to the consumed result pointer. Should use MoveTemp.
  * @return				Returns the final result, which is either the combined result pointer and new result, or just the new result
  */
-static FNetResult AddToAndConsumeChainResultPtr(TPimplPtr<FNetResult, EPimplPtrMode::DeepCopy>& ResultPtr, FNetResult&& Result)
+inline FNetResult AddToAndConsumeChainResultPtr(TPimplPtr<FNetResult, EPimplPtrMode::DeepCopy>& ResultPtr, FNetResult&& Result)
 {
 	const bool bValidResultPtr = ResultPtr.IsValid();
 	FNetResult ReturnVal(bValidResultPtr ? MoveTemp(*ResultPtr) : MoveTemp(Result));
@@ -451,31 +451,6 @@ public:
 		return !(TNetResult::operator == (A));
 	}
 
-
-	/**
-	 * Chains a new result to the result pointer, putting it at the end of the chain - or assigns it to the result pointer, if not set.
-	 *
-	 * @param ResultPtr		The result pointer to chain or assign the new result to.
-	 * @param Result		The result to chain to the result pointer. Should use MoveTemp.
-	 */
-	inline friend void AddToChainResultPtr(TPimplPtr<FNetResult, EPimplPtrMode::DeepCopy>& ResultPtr, TNetResult&& Result)
-	{
-		AddToChainResultPtr(ResultPtr, static_cast<FNetResult&&>(Result));
-	}
-
-	/**
-	 * Chains a new result to the end of the result pointer (or takes its place if not set), then consumes/moves/returns the pointer result
-	 *
-	 * @param ResultPtr		The result pointer to consume and chain the new result to.
-	 * @param Result		The result to chain to the consumed result pointer. Should use MoveTemp.
-	 * @return				Returns the final result, which is either the combined result pointer and new result, or just the new result
-	 */
-	inline friend FNetResult AddToAndConsumeChainResultPtr(TPimplPtr<FNetResult, EPimplPtrMode::DeepCopy>& ResultPtr, TNetResult&& Result)
-	{
-		return AddToAndConsumeChainResultPtr(ResultPtr, static_cast<FNetResult&&>(Result));
-	}
-
-
 private:
 	/** Default ErrorContext based on ResultEnum */
 	static const TCHAR* DefaultErrorContext(ResultEnum InResult)
@@ -496,8 +471,8 @@ private:
 	{ \
 	namespace Net \
 	{ \
-		void AddToChainResultPtr(TPimplPtr<FNetResult, EPimplPtrMode::DeepCopy>& ResultPtr, TNetResult<EnumType>&& Result); \
-		FNetResult AddToAndConsumeChainResultPtr(TPimplPtr<FNetResult, EPimplPtrMode::DeepCopy>& ResultPtr, TNetResult<EnumType>&& Result); \
+		inline void AddToChainResultPtr(TPimplPtr<FNetResult, EPimplPtrMode::DeepCopy>& ResultPtr, TNetResult<EnumType>&& Result) { AddToChainResultPtr(ResultPtr, static_cast<FNetResult&&>(Result)); } \
+		inline FNetResult AddToAndConsumeChainResultPtr(TPimplPtr<FNetResult, EPimplPtrMode::DeepCopy>& ResultPtr, TNetResult<EnumType>&& Result) { return AddToAndConsumeChainResultPtr(ResultPtr, static_cast<FNetResult&&>(Result)); } \
 	} \
 	}
 
