NetBitArray
* Added ForAllSetBitsMultiple to call a functor over set bits calculated from : A Op1 (B Op2 C)

#rnx

diff --git a/Engine/Plugins/Runtime/ReplicationSystemTestPlugin/Source/Private/Tests/Core/TestNetBitArray.cpp b/Engine/Plugins/Runtime/ReplicationSystemTestPlugin/Source/Private/Tests/Core/TestNetBitArray.cpp
--- a/Engine/Plugins/Runtime/ReplicationSystemTestPlugin/Source/Private/Tests/Core/TestNetBitArray.cpp
+++ b/Engine/Plugins/Runtime/ReplicationSystemTestPlugin/Source/Private/Tests/Core/TestNetBitArray.cpp
@@ -993,6 +993,55 @@ UE_NET_TEST(FNetBitArrayView, TestForEachSet)
 	}
 }
 
+UE_NET_TEST(FNetBitArrayView, TestForAllSetBitsMultiple)
+{	
+	struct FCollectSetBitsFunctor
+	{
+		void operator()(uint32 Index) { Invoked.Add(Index); }
+		TArray<uint32> Invoked;
+	};
+
+	auto TestOp = []<typename WordOp1, typename WordOp2>(WordOp1 Op1, WordOp2 Op2) -> TTuple<int32, int32>
+	{
+		uint32 BufferA = 0xFFFFFFFF;
+		FNetBitArrayView BitArrayA(&BufferA, 32, FNetBitArrayView::NoResetNoValidate);
+
+		uint32 BufferB = 0xAAAAAAAA;
+		FNetBitArrayView BitArrayB(&BufferB, 32, FNetBitArrayView::NoResetNoValidate);
+
+		uint32 BufferC = 0x00000000;
+		FNetBitArrayView BitArrayC(&BufferC, 32, FNetBitArrayView::NoResetNoValidate);
+
+		FCollectSetBitsFunctor Functor;
+		FNetBitArrayView::ForAllSetBitsMultiple(BitArrayA, Op1, BitArrayB, Op2, BitArrayC, Functor);
+
+		const uint32 FinalWord = Op1(BufferA, Op2(BufferB, BufferC));
+		const int32 ExpectedBits = FMath::CountBits(FinalWord);
+		
+		return TTuple<int32, int32>(ExpectedBits, Functor.Invoked.Num());
+	};
+
+	{
+		auto Result = TestOp(FNetBitArrayView::AndOp, FNetBitArrayView::OrOp);
+		UE_NET_ASSERT_EQ(Result.Get<0>(), Result.Get<1>());
+	}
+
+	{
+		auto Result = TestOp(FNetBitArrayView::AndNotOp, FNetBitArrayView::OrOp);
+		UE_NET_ASSERT_EQ(Result.Get<0>(), Result.Get<1>());
+	}
+
+	{
+		auto Result = TestOp(FNetBitArrayView::XorOp, FNetBitArrayView::OrOp);
+		UE_NET_ASSERT_EQ(Result.Get<0>(), Result.Get<1>());
+	}
+
+	{
+		auto Result = TestOp(FNetBitArrayView::OrOp, FNetBitArrayView::AndOp);
+		UE_NET_ASSERT_EQ(Result.Get<0>(), Result.Get<1>());
+	}
+}
+
 UE_NET_TEST(FNetBitArrayView, TestGetSetBitIndices)
 {	
 	constexpr uint32 OOBValue = ~0U;
diff --git a/Engine/Source/Runtime/Net/Core/Public/Net/Core/NetBitArray.h b/Engine/Source/Runtime/Net/Core/Public/Net/Core/NetBitArray.h
--- a/Engine/Source/Runtime/Net/Core/Public/Net/Core/NetBitArray.h
+++ b/Engine/Source/Runtime/Net/Core/Public/Net/Core/NetBitArray.h
@@ -293,6 +293,10 @@ public:
 	template<typename T, typename V>
 	static void ForAllSetBits(const FNetBitArray& A, const FNetBitArray& B, T&& WordOpFunctor, V&& Functor);
 
+	/** Call the functor for every bit still set after the following operation: "A Op1 (B Op2 C)" */
+	template <typename WordOpFunctor1, typename WordOpFunctor2, typename UserFunctor>
+	static void ForAllSetBitsMultiple(const FNetBitArray& A, WordOpFunctor1&& Op1, const FNetBitArray& B, WordOpFunctor2&& Op2, const FNetBitArray& C, UserFunctor&& Functor);
+
 	/** Compare two BitArrays bit by bit and invoke FunctorA(BitIndex) for bits set only in array A and invoke FunctorB(BitIndex) for bits set only in array B. */
 	template<typename T, typename V>
 	static void ForAllExclusiveBits(const FNetBitArray& A, const FNetBitArray& B, T&& FunctorA, V&& FunctorB);
@@ -521,6 +525,10 @@ public:
 	template <typename T, typename V>
 	static void ForAllSetBits(const FNetBitArrayView& A, const FNetBitArrayView& B, T&& WordOpFunctor, V&& Functor);
 
+	/** Call the functor for every bit still set after the following operation: "A Op1 (B Op2 C)" */
+	template <typename WordOpFunctor1, typename WordOpFunctor2, typename UserFunctor>
+	static void ForAllSetBitsMultiple(const FNetBitArrayView& A, WordOpFunctor1&& Op1, const FNetBitArrayView& B, WordOpFunctor2&& Op2, const FNetBitArrayView& C, UserFunctor&& Functor);
+
 	/** Compare two BitArrays bit by bit and invoke FunctorA(BitIndex) for bits set only in array A and invoke FunctorB(BitIndex) for bits set only in array B. */
 	template <typename T, typename V>
 	static void ForAllExclusiveBits(const FNetBitArrayView& A, const FNetBitArrayView& B, T&& FunctorA, V&& FunctorB);
@@ -941,6 +949,33 @@ private:
 		}
 	}
 
+	template <typename WordOpFunctor1, typename WordOpFunctor2, typename UserFunctor> static void ForAllSetBitsMultiple(const StorageWordType* StorageA, const StorageWordType* StorageB, const StorageWordType* StorageC, const uint32 WordCount, const uint32 BitCount, WordOpFunctor1&& WordOp1, WordOpFunctor2&& WordOp2, UserFunctor&& Functor)
+	{
+		const StorageWordType LastWordMask = ~StorageWordType(0) >> (uint32(-int32(BitCount)) & (WordBitCount - 1));
+		const uint32 LastWordIt = WordCount - 1;
+
+		for (uint32 WordIt = 0, CurrentBitIndex = 0; WordIt < WordCount; ++WordIt, CurrentBitIndex += WordBitCount)
+		{
+			const StorageWordType CurrentWordA = StorageA[WordIt];
+			const StorageWordType CurrentWordB = StorageB[WordIt];
+			const StorageWordType CurrentWordC = StorageC[WordIt];
+
+			StorageWordType CurrentWord = WordOp1(CurrentWordA, WordOp2(CurrentWordB, CurrentWordC) ) & ((WordIt == LastWordIt) ? LastWordMask : ~0U);
+
+			// Test All bits in the CurrentWord and invoke functor if they are set
+			uint32 LocalBitIndex = CurrentBitIndex;
+			while (CurrentWord)
+			{
+				if (CurrentWord & 0x1)
+				{
+					Functor(LocalBitIndex);
+				}
+				CurrentWord >>= 1;
+				++LocalBitIndex;
+			}
+		}
+	}
+
 	template<typename T, typename V, typename Predicate> static void ForAllExclusiveBitsByPredicate(const StorageWordType* StorageA, const StorageWordType* StorageB, const uint32 WordCount, const uint32 BitCount, T&& FunctorA, V&& FunctorB, Predicate&& Pred)
 	{
 		const StorageWordType LastWordMask = ~StorageWordType(0) >> (uint32(-int32(BitCount)) & (WordBitCount - 1));
@@ -1286,6 +1321,13 @@ template <typename T, typename V> inline void FNetBitArray::ForAllSetBits(const
 	FNetBitArrayHelper::ForAllSetBits(A.Storage.GetData(), B.Storage.GetData(), A.Storage.Num(), A.BitCount, WordOpFunctor, Functor);
 }
 
+template<typename WordOpFunctor1, typename WordOpFunctor2, typename UserFunctor> inline void FNetBitArray::ForAllSetBitsMultiple(const FNetBitArray& A, WordOpFunctor1&& Op1, const FNetBitArray& B, WordOpFunctor2&& Op2, const FNetBitArray& C, UserFunctor&& Functor)
+{
+	UE_NETBITARRAY_VALIDATE_BOTH_COMPATIBLE(A, B);
+	UE_NETBITARRAY_VALIDATE_BOTH_COMPATIBLE(B, C);
+	FNetBitArrayHelper::ForAllSetBitsMultiple(A.Storage.GetData(), B.Storage.GetData(), C.Storage.GetData(), A.Storage.Num(), A.BitCount, Op1, Op2, Functor);
+}
+
 template <typename T, typename V> inline void FNetBitArray::ForAllExclusiveBits(const FNetBitArray& A, const FNetBitArray& B, T&& FunctorA, V&& FunctorB)
 {
 	UE_NETBITARRAY_CHECK(A.GetNumBits() == B.GetNumBits());
@@ -1537,19 +1579,26 @@ template<typename T> inline void FNetBitArrayView::ForAllSetBits(T&& Functor) co
 	FNetBitArrayHelper::ForAllSetBits(Storage, WordCount, BitCount, Functor);
 }
 
-
 template<typename T> inline void FNetBitArrayView::ForAllUnsetBits(T&& Functor) const
 {
 	FNetBitArrayHelper::ForAllUnsetBits(Storage, WordCount, BitCount, Functor);
 }
 
-template <typename T, typename V>
-void FNetBitArrayView::ForAllSetBits(const FNetBitArrayView& A, const FNetBitArrayView& B, T&& WordOpFunctor, V&& Functor)
+template <typename T, typename V> 
+inline void FNetBitArrayView::ForAllSetBits(const FNetBitArrayView& A, const FNetBitArrayView& B, T&& WordOpFunctor, V&& Functor)
 {
 	UE_NETBITARRAY_CHECK(A.GetNumBits() == B.GetNumBits());
 	FNetBitArrayHelper::ForAllSetBits(A.Storage, B.Storage, A.WordCount, A.BitCount, WordOpFunctor, Functor);
 }
 
+template <typename WordOpFunctor1, typename WordOpFunctor2, typename UserFunctor> 
+inline void FNetBitArrayView::ForAllSetBitsMultiple(const FNetBitArrayView& A, WordOpFunctor1&& Op1, const FNetBitArrayView& B, WordOpFunctor2&& Op2, const FNetBitArrayView& C, UserFunctor&& Functor)
+{
+	UE_NETBITARRAY_VALIDATE_BOTH_COMPATIBLE(A, B);
+	UE_NETBITARRAY_VALIDATE_BOTH_COMPATIBLE(B, C);
+	FNetBitArrayHelper::ForAllSetBitsMultiple(A.Storage, B.Storage, C.Storage, A.WordCount, A.BitCount, Op1, Op2, Functor);
+}
+
 template <typename T, typename V>
 void FNetBitArrayView::ForAllExclusiveBits(const FNetBitArrayView& A, const FNetBitArrayView& B, T&& FunctorA, V&& FunctorB)
 {
