Network
* Added UE_FAST_ARRAY_COMPILE_LOG_LEVEL that allows a program to control the compile time verbosity of LogNetFastTArray logs
* Since this class is templated the logs in the functions are unique per implemented class so it can grow the amount of the log static strings if they are left compiled-in

diff --git a/Engine/Source/Runtime/Net/Core/Classes/Net/Serialization/FastArraySerializer.h b/Engine/Source/Runtime/Net/Core/Classes/Net/Serialization/FastArraySerializer.h
--- a/Engine/Source/Runtime/Net/Core/Classes/Net/Serialization/FastArraySerializer.h
+++ b/Engine/Source/Runtime/Net/Core/Classes/Net/Serialization/FastArraySerializer.h
@@ -28,7 +28,12 @@ class Error;
 struct FFastArraySerializer;
 struct FFastArraySerializerItem;
 
-NETCORE_API DECLARE_LOG_CATEGORY_EXTERN(LogNetFastTArray, Warning, All);
+// Override this to define the compiled log level of LogNetFastTArray and reduce the amount of logs generated with each FastArray implementation
+#ifndef UE_FAST_ARRAY_COMPILE_LOG_LEVEL
+	#define UE_FAST_ARRAY_COMPILE_LOG_LEVEL All
+#endif
+
+NETCORE_API DECLARE_LOG_CATEGORY_EXTERN(LogNetFastTArray, Warning, UE_FAST_ARRAY_COMPILE_LOG_LEVEL);
 
 DECLARE_CYCLE_STAT_EXTERN(TEXT("NetSerializeFast Array"), STAT_NetSerializeFastArray, STATGROUP_ServerCPU, NETCORE_API);
 DECLARE_CYCLE_STAT_EXTERN(TEXT("NetSerializeFast Array BuildMap"), STAT_NetSerializeFastArray_BuildMap, STATGROUP_ServerCPU, NETCORE_API);
