Network
* Added UE_FAST_ARRAY_RUNTIME_LOG_LEVEL so the runtime log level can be set less than or equal to UE_FAST_ARRAY_COMPILE_LOG_LEVEL to prevent compiler errors.

diff --git a/Engine/Source/Runtime/Net/Core/Classes/Net/Serialization/FastArraySerializer.h b/Engine/Source/Runtime/Net/Core/Classes/Net/Serialization/FastArraySerializer.h
--- a/Engine/Source/Runtime/Net/Core/Classes/Net/Serialization/FastArraySerializer.h
+++ b/Engine/Source/Runtime/Net/Core/Classes/Net/Serialization/FastArraySerializer.h
@@ -33,7 +33,12 @@ struct FFastArraySerializerItem;
 	#define UE_FAST_ARRAY_COMPILE_LOG_LEVEL All
 #endif
 
-NETCORE_API DECLARE_LOG_CATEGORY_EXTERN(LogNetFastTArray, Warning, UE_FAST_ARRAY_COMPILE_LOG_LEVEL);
+// Runtime log level must be <= compiled log level
+#ifndef UE_FAST_ARRAY_RUNTIME_LOG_LEVEL
+	#define UE_FAST_ARRAY_RUNTIME_LOG_LEVEL Warning
+#endif
+
+NETCORE_API DECLARE_LOG_CATEGORY_EXTERN(LogNetFastTArray, UE_FAST_ARRAY_RUNTIME_LOG_LEVEL, UE_FAST_ARRAY_COMPILE_LOG_LEVEL);
 
 DECLARE_CYCLE_STAT_EXTERN(TEXT("NetSerializeFast Array"), STAT_NetSerializeFastArray, STATGROUP_ServerCPU, NETCORE_API);
 DECLARE_CYCLE_STAT_EXTERN(TEXT("NetSerializeFast Array BuildMap"), STAT_NetSerializeFastArray_BuildMap, STATGROUP_ServerCPU, NETCORE_API);
