Do net trace initialisation in OnConnected rather than OnTraceStarted so that it works when the tail buffer is disabled or full

diff --git a/Engine/Source/Runtime/Net/Core/Private/Net/Core/Trace/NetTraceInternal.cpp b/Engine/Source/Runtime/Net/Core/Private/Net/Core/Trace/NetTraceInternal.cpp
--- a/Engine/Source/Runtime/Net/Core/Private/Net/Core/Trace/NetTraceInternal.cpp
+++ b/Engine/Source/Runtime/Net/Core/Private/Net/Core/Trace/NetTraceInternal.cpp
@@ -3,6 +3,7 @@
 #include "AutoRTFM.h"
 #include "Containers/ArrayView.h"
 #include "Containers/Map.h"
+#include "Containers/Ticker.h"
 #include "Net/Core/Trace/NetDebugName.h"
 #include "Net/Core/Trace/Reporters/NetTraceReporter.h"
 #include "Hash/CityHash.h"
@@ -82,7 +83,7 @@ struct FNetTraceInternal
 static std::atomic<uint32> CurrentTraceId(0);
 static std::atomic<UE::Net::FNetDebugNameId> NextNameId(1);
 static thread_local TUniquePtr<FNetTraceInternal::FThreadBuffer> ThreadBuffer;
-static FDelegateHandle OnTraceStartedHandle;
+static FDelegateHandle OnTraceConnectedHandle;
 static FDelegateHandle OnTraceStoppedHandle;
 static TMap<uint32, FNetTraceInternal::FGameInstance> GameInstances;
 
@@ -113,7 +114,7 @@ UE_AUTORTFM_ALWAYS_OPEN uint32 FNetTraceInternal::GetCurrentTraceId()
 	return CurrentTraceId;
 }
 
-void FNetTrace::OnTraceStarted(FTraceAuxiliary::EConnectionType /*TraceType*/, const FString& /*TraceDestination*/)
+void FNetTrace::OnTraceConnected()
 {
 	++CurrentTraceId;
 	NextNameId = 1;
@@ -165,9 +166,17 @@ void FNetTrace::SetTraceVerbosity(uint32 Verbosity)
 
 		if (UE::Trace::IsTracing())
 		{
-			OnTraceStarted(FTraceAuxiliary::EConnectionType::None, TEXT(""));
+			OnTraceConnected();
 		}
-		OnTraceStartedHandle = FTraceAuxiliary::OnTraceStarted.AddStatic(&FNetTrace::OnTraceStarted);
+
+		OnTraceConnectedHandle = FTraceAuxiliary::OnConnection.AddLambda([]() 
+		{
+			// OnConnection will be called from a worker thread
+			ExecuteOnGameThread(TEXT("FNetTrace::OnTraceConnected"), []()
+			{
+				FNetTrace::OnTraceConnected();
+			});
+		});
 		OnTraceStoppedHandle = FTraceAuxiliary::OnTraceStopped.AddStatic(&FNetTrace::OnTraceStopped);
 	}
 	else if (GetTraceVerbosity() && !NewVerbosity)
@@ -178,7 +187,7 @@ void FNetTrace::SetTraceVerbosity(uint32 Verbosity)
 		{
 			OnTraceStopped(FTraceAuxiliary::EConnectionType::None, TEXT(""));
 		}
-		FTraceAuxiliary::OnTraceStarted.Remove(OnTraceStartedHandle);
+		FTraceAuxiliary::OnConnection.Remove(OnTraceConnectedHandle);
 		FTraceAuxiliary::OnTraceStopped.Remove(OnTraceStoppedHandle);
 	}
 
diff --git a/Engine/Source/Runtime/Net/Core/Public/Net/Core/Trace/Private/NetTraceInternal.h b/Engine/Source/Runtime/Net/Core/Public/Net/Core/Trace/Private/NetTraceInternal.h
--- a/Engine/Source/Runtime/Net/Core/Public/Net/Core/Trace/Private/NetTraceInternal.h
+++ b/Engine/Source/Runtime/Net/Core/Public/Net/Core/Trace/Private/NetTraceInternal.h
@@ -186,7 +186,7 @@ struct FNetTrace
 	/** Sets the runtime trace verbosity level, Verbosity == 0 means that it is disabled */
 	NETCORE_API static void SetTraceVerbosity(uint32 Verbosity);
 
-	static void OnTraceStarted(FTraceAuxiliary::EConnectionType TraceType, const FString& TraceDestination);
+	static void OnTraceConnected();
 
 	static void OnTraceStopped(FTraceAuxiliary::EConnectionType TraceType, const FString& TraceDestination);
 
