Fix missing case in NetEnums.h

#rnx

diff --git a/Engine/Source/Runtime/Net/Core/Public/Net/Core/Connection/NetEnums.h b/Engine/Source/Runtime/Net/Core/Public/Net/Core/Connection/NetEnums.h
--- a/Engine/Source/Runtime/Net/Core/Public/Net/Core/Connection/NetEnums.h
+++ b/Engine/Source/Runtime/Net/Core/Public/Net/Core/Connection/NetEnums.h
@@ -139,11 +139,30 @@ inline const TCHAR* LexToString(ENetRefHandleError NetRefHandleError)
 {
 	switch (NetRefHandleError)
 	{
-		case ENetRefHandleError::None: return TEXT("None");
-		case ENetRefHandleError::BitstreamCorrupted: return TEXT("BitstreamCorrupted");
-		case ENetRefHandleError::ReplicationDisabled: return TEXT("ReplicationDisabled");
-		case ENetRefHandleError::BlockedByMustBeMapped: return TEXT("BlockedByMustBeMapped");
-		default: return TEXT("missing");
+		case ENetRefHandleError::None:
+		{
+			return TEXT("None");
+		} break;
+		case ENetRefHandleError::BitstreamCorrupted:
+		{
+			return TEXT("BitstreamCorrupted");
+		} break;
+		case ENetRefHandleError::ReplicationDisabled:
+		{
+			return TEXT("ReplicationDisabled");
+		} break;
+		case ENetRefHandleError::BlockedByMustBeMapped:
+		{
+			return TEXT("BlockedByMustBeMapped");
+		} break;
+		case ENetRefHandleError::BlockedByCreationParents:
+		{
+			return TEXT("BlockedByCreationParents");
+		} break;
+		default:
+		{
+			return TEXT("missing");
+		}
 	}
 }
 
