Fix to allow Polling NetTracing across multiple threads. The ObjectsTraced list need to be moved into Thead Local Storage - this can result in the same objects being stored per-thread, but this should only result in slightly higher memory usage and not break the trace file.
One thing to note - due to a bug in the engine Trace code, we need to use NoSync on the PollingEvent or the trace could end up with a missing Sync event and make the file unreadable from that point on.

#rnx
#tests Tested against Dedicated Server with Cooked Data Win64

diff --git a/Engine/Source/Runtime/Experimental/Iris/Core/Private/Iris/ReplicationSystem/Polling/ObjectPoller.cpp b/Engine/Source/Runtime/Experimental/Iris/Core/Private/Iris/ReplicationSystem/Polling/ObjectPoller.cpp
--- a/Engine/Source/Runtime/Experimental/Iris/Core/Private/Iris/ReplicationSystem/Polling/ObjectPoller.cpp
+++ b/Engine/Source/Runtime/Experimental/Iris/Core/Private/Iris/ReplicationSystem/Polling/ObjectPoller.cpp
@@ -314,7 +314,7 @@ void FObjectPoller::ForcePollObject(FInternalNetRefIndex ObjectIndex, FNetStatsC
 	{
 		IRIS_PROFILER_PROTOCOL_NAME_CONDITIONAL(ObjectData.Protocol->DebugName->Name, CVars::bEnableVerbosePollProfiling);
 		UE_NET_IRIS_STATS_TIMER(Timer, InNetStatsContext);
-		//UE_NET_TRACE_POLL_OBJECT_SCOPE(ObjectData.RefHandle, Timer); //Temporarily disable to fix a crash
+		UE_NET_TRACE_POLL_OBJECT_SCOPE(ObjectData.RefHandle, Timer);
 
 		const bool bIsGCAffectedObject = GarbageCollectionAffectedObjects.GetBit(ObjectIndex);
 		GarbageCollectionAffectedObjects.ClearBit(ObjectIndex);
@@ -336,7 +336,7 @@ void FObjectPoller::ForcePollObject(FInternalNetRefIndex ObjectIndex, FNetStatsC
 		else
 		{
 			UE_NET_IRIS_STATS_ADD_TIME_AND_COUNT_FOR_OBJECT_AS_WASTE(Timer, Poll, ObjectIndex);
-			//UE_NET_TRACE_POLL_OBJECT_IS_WASTE(); //Temporarily disable to fix a crash
+			UE_NET_TRACE_POLL_OBJECT_IS_WASTE();
 		}
 		++InPollStats.PolledObjectCount;
 	}
@@ -388,7 +388,7 @@ void FObjectPoller::PushModelPollObject(FInternalNetRefIndex ObjectIndex, FNetSt
 		{
 			IRIS_PROFILER_PROTOCOL_NAME_CONDITIONAL(ObjectData.Protocol->DebugName->Name, CVars::bEnableVerbosePollProfiling);
 			UE_NET_IRIS_STATS_TIMER(Timer, NetStatsContext);
-			//UE_NET_TRACE_POLL_OBJECT_SCOPE(ObjectData.RefHandle, Timer); //Temporarily disable to fix a crash
+			UE_NET_TRACE_POLL_OBJECT_SCOPE(ObjectData.RefHandle, Timer);
 
 			// We need to do a poll if object is marked as dirty
 			EReplicationFragmentPollFlags PollOptions = bUsePerPropertyDirtyTracking && !bWantsFullPoll ? EReplicationFragmentPollFlags::PollDirtyState : EReplicationFragmentPollFlags::PollAllState;
@@ -403,7 +403,7 @@ void FObjectPoller::PushModelPollObject(FInternalNetRefIndex ObjectIndex, FNetSt
 		{
 			IRIS_PROFILER_PROTOCOL_NAME_CONDITIONAL(ObjectData.Protocol->DebugName->Name, CVars::bEnableVerbosePollProfiling);
 			UE_NET_IRIS_STATS_TIMER(Timer, NetStatsContext);
-			//UE_NET_TRACE_POLL_OBJECT_SCOPE(ObjectData.RefHandle, Timer); //Temporarily disable to fix a crash
+			UE_NET_TRACE_POLL_OBJECT_SCOPE(ObjectData.RefHandle, Timer);
 
 			// If this object might have been affected by GC, only refresh cached references
 			bPollFoundDirty = FReplicationInstanceOperations::PollAndCopyObjectReferences(InstanceProtocol, EReplicationFragmentTraits::None);
@@ -420,7 +420,7 @@ void FObjectPoller::PushModelPollObject(FInternalNetRefIndex ObjectIndex, FNetSt
 	{
 		IRIS_PROFILER_PROTOCOL_NAME_CONDITIONAL(ObjectData.Protocol->DebugName->Name, CVars::bEnableVerbosePollProfiling);
 		UE_NET_IRIS_STATS_TIMER(Timer, NetStatsContext);
-		//UE_NET_TRACE_POLL_OBJECT_SCOPE(ObjectData.RefHandle, Timer); //Temporarily disable to fix a crash
+		UE_NET_TRACE_POLL_OBJECT_SCOPE(ObjectData.RefHandle, Timer);
 
 		// If the object has fragments with pushed based properties, and is not marked dirty and object is affected by GC we need to make sure that we refresh cached references for all fragments with push based properties
 		const bool bIsFullPushBasedObject = EnumHasAnyFlags(InstanceTraits, EReplicationInstanceProtocolTraits::HasFullPushBasedDirtiness);
@@ -457,7 +457,7 @@ void FObjectPoller::PushModelPollObject(FInternalNetRefIndex ObjectIndex, FNetSt
 		else
 		{
 			UE_NET_IRIS_STATS_ADD_TIME_AND_COUNT_FOR_OBJECT_AS_WASTE(Timer, Poll, ObjectIndex);
-			//UE_NET_TRACE_POLL_OBJECT_IS_WASTE(); //Temporarily disable to fix a crash
+			UE_NET_TRACE_POLL_OBJECT_IS_WASTE();
 		}
 	}
 }
diff --git a/Engine/Source/Runtime/Net/Core/Private/Net/Core/Trace/NetTraceInternal.cpp b/Engine/Source/Runtime/Net/Core/Private/Net/Core/Trace/NetTraceInternal.cpp
--- a/Engine/Source/Runtime/Net/Core/Private/Net/Core/Trace/NetTraceInternal.cpp
+++ b/Engine/Source/Runtime/Net/Core/Private/Net/Core/Trace/NetTraceInternal.cpp
@@ -33,18 +33,6 @@ struct FNetTraceInternal
 		ENetTraceVersion_DebugNameIndexIs32Bits = 5,
 	};
 
-	struct FThreadBuffer
-	{
-		// Trace Id allows us to discard the old Name Id maps when we start a new trace
-		uint32 TraceId;
-
-		// Map FName to NameId, used to know what FNames we have traced
-		TMap<FName, UE::Net::FNetDebugNameId> DynamicFNameToNameIdMap;
-
-		// Map hashed dynamic TCHAR* strings to NameId
-		TMap<uint64, UE::Net::FNetDebugNameId> DynamicNameHashToNameIdMap;
-	};
-
 	struct FConnection
 	{
 		FString Address;
@@ -65,17 +53,40 @@ struct FNetTraceInternal
 	{
 		TMap<uint32, FConnection> Connections;
 		TMap<uint64, FNetObject> Objects;
-		TSet<uint64> ObjectsTraced;
 		FString Name;
 		bool bIsServer;
 	};
 
+	//Per-GameInstance variables that are stored per-thread
+	struct FObjectTraceCounter
+	{
+		TSet<uint64> ObjectsTraced;
+	};
+
+	struct FThreadBuffer
+	{
+		// Trace Id allows us to discard the old Name Id maps when we start a new trace
+		uint32 TraceId;
+
+		// Map FName to NameId, used to know what FNames we have traced
+		TMap<FName, UE::Net::FNetDebugNameId> DynamicFNameToNameIdMap;
+
+		// Map hashed dynamic TCHAR* strings to NameId
+		TMap<uint64, UE::Net::FNetDebugNameId> DynamicNameHashToNameIdMap;
+
+		// Stores the Objects and ObjectsTraces structures on a Per-Game-Instance
+		TMap<uint32, FObjectTraceCounter> GameInstanceToObjectTraceCounterMap;
+
+		// Map FNetDebugName pointer to NameId - we can never move memory location of the Persistent FNetDebugNames
+		TMap<const UE::Net::FNetDebugName*, UE::Net::FNetDebugNameId> DebugNameToNameIdMap;
+	};
+
 	// Get next NameId used to track what we already have traced
 	static UE::Net::FNetDebugNameId GetNextNameId();
 
 	static uint32 GetCurrentTraceId();
 
-	static inline FThreadBuffer* GetThreadBuffer();
+	static inline FThreadBuffer& GetThreadBuffer();
 
 	static constexpr ENetTraceVersion NetTraceVersion = ENetTraceVersion::ENetTraceVersion_DebugNameIndexIs32Bits;
 };
@@ -90,7 +101,7 @@ static TMap<uint32, FNetTraceInternal::FGameInstance> GameInstances;
 FNetTrace::FOnResetPersistentNetDebugNames FNetTrace::OnResetPersistentNetDebugNames;
 FNetTrace::FOnNetTraceStarted FNetTrace::OnNetTraceStarted;
 
-FNetTraceInternal::FThreadBuffer* FNetTraceInternal::GetThreadBuffer()
+FNetTraceInternal::FThreadBuffer& FNetTraceInternal::GetThreadBuffer()
 {
 	FNetTraceInternal::FThreadBuffer* ThreadBufferPtr = ThreadBuffer.Get();
 	uint32 TraceId = GetCurrentTraceId();
@@ -101,7 +112,9 @@ FNetTraceInternal::FThreadBuffer* FNetTraceInternal::GetThreadBuffer()
 		ThreadBufferPtr->TraceId = TraceId;
 	}
 
-	return ThreadBufferPtr;
+	check(ThreadBufferPtr);
+
+	return *ThreadBufferPtr;
 }
 
 UE_AUTORTFM_ALWAYS_OPEN UE::Net::FNetDebugNameId FNetTraceInternal::GetNextNameId()
@@ -128,8 +141,6 @@ void FNetTrace::OnTraceConnected()
 		const uint32 GameInstanceId = InstanceIter.Key;
 		FNetTraceInternal::FGameInstance& GameInstance = InstanceIter.Value;
 
-		GameInstance.ObjectsTraced.Reset();
-
 		FNetTraceInternal::Reporter::ReportInstanceUpdated(GameInstanceId, GameInstance.bIsServer, *GameInstance.Name);
 
 		for (const TPair<uint32, FNetTraceInternal::FConnection>& ConnectionIter : GameInstance.Connections)
@@ -489,6 +500,18 @@ void FNetTrace::PopSendBunch(FNetTraceCollector& Collector)
 	BunchHeaderEvent.EndPos = 0U;
 }
 
+static UE::Net::FNetDebugNameId& GetOrCreateDebugNameIdInLocalMap(const UE::Net::FNetDebugName* DebugNamePtr, FNetTraceInternal::FThreadBuffer& InThreadBuffer)
+{
+	UE::Net::FNetDebugNameId* NameIdPtr = InThreadBuffer.DebugNameToNameIdMap.Find(DebugNamePtr);
+	if(!NameIdPtr)
+	{
+		NameIdPtr = &(InThreadBuffer.DebugNameToNameIdMap.Add(DebugNamePtr, FNetTrace::TraceName(DebugNamePtr->Name)));
+	}
+
+	check(NameIdPtr);
+	return *NameIdPtr;
+}
+
 static void TraceExistingObject(uint32 GameInstanceId, const TMap<uint64, FNetTraceInternal::FNetObject>& Objects, TSet<uint64>& ObjectsTraced, uint64 NetObjectId)
 {
 	if (!ObjectsTraced.Contains(NetObjectId))
@@ -500,12 +523,10 @@ static void TraceExistingObject(uint32 GameInstanceId, const TMap<uint64, FNetTr
 
 			if (NetObject->DebugName)
 			{
-				if (NetObject->DebugName->DebugNameId == 0U)
-				{
-					NetObject->DebugName->DebugNameId = FNetTrace::TraceName(NetObject->DebugName->Name);
-				}
+				FNetTraceInternal::FThreadBuffer& LocalThreadBuffer = FNetTraceInternal::GetThreadBuffer();
+				const UE::Net::FNetDebugNameId DebugNameId = GetOrCreateDebugNameIdInLocalMap(NetObject->DebugName, LocalThreadBuffer);
 					
-				FNetTraceInternal::Reporter::ReportObjectExists(GameInstanceId, NetObjectId, NetObject->DebugName->DebugNameId, NetObject->TypeIdentifier, NetObject->OwnerId);
+				FNetTraceInternal::Reporter::ReportObjectExists(GameInstanceId, NetObjectId, DebugNameId, NetObject->TypeIdentifier, NetObject->OwnerId);
 			}
 			else
 			{
@@ -530,9 +551,11 @@ void FNetTrace::TraceCollectedEvents(FNetTraceCollector& Collector, uint32 GameI
 	PacketInfo.PacketType = PacketType;
 
 	FNetTraceInternal::FGameInstance& GameInstance = GameInstances.FindOrAdd(GameInstanceId);
+	FNetTraceInternal::FThreadBuffer& LocalThreadBuffer = FNetTraceInternal::GetThreadBuffer();
+	FNetTraceInternal::FObjectTraceCounter& ObjectTraceData = LocalThreadBuffer.GameInstanceToObjectTraceCounterMap.FindOrAdd(GameInstanceId);
 	for (uint64 NetObjectId : Collector.ReferencedObjectIds)
 	{
-		TraceExistingObject(GameInstanceId, GameInstance.Objects, GameInstance.ObjectsTraced, NetObjectId);
+		TraceExistingObject(GameInstanceId, GameInstance.Objects, ObjectTraceData.ObjectsTraced, NetObjectId);
 	}
 
 	// Trace all collected events
@@ -573,14 +596,20 @@ void FNetTrace::TracePacket(uint32 GameInstanceId, uint32 ConnectionId, uint32 P
 
 void FNetTrace::TraceObjectCreated(uint32 GameInstanceId, uint64 NetObjectId, const UE::Net::FNetDebugName* DebugName, uint64 TypeIdentifier, uint32 OwnerId)
 {
-	GameInstances.FindOrAdd(GameInstanceId).Objects.Emplace(NetObjectId, { .Name = NAME_None, .DebugName = DebugName, .TypeIdentifier = TypeIdentifier, .ParentNetObjectId = uint64(-1), .OwnerId = OwnerId });
+	FNetTraceInternal::FGameInstance& GameInstance = GameInstances.FindOrAdd(GameInstanceId);
+	FNetTraceInternal::FThreadBuffer& LocalThreadBuffer = FNetTraceInternal::GetThreadBuffer();
+	FNetTraceInternal::FObjectTraceCounter& ObjectTraceData = LocalThreadBuffer.GameInstanceToObjectTraceCounterMap.FindOrAdd(GameInstanceId);
+
+	GameInstance.Objects.Emplace(NetObjectId, { .Name = NAME_None, .DebugName = DebugName, .TypeIdentifier = TypeIdentifier, .ParentNetObjectId = uint64(-1), .OwnerId = OwnerId });
 
 	if (!GNetTraceRuntimeVerbosity)
 	{
 		return;
 	}
 
-	GameInstances[GameInstanceId].ObjectsTraced.Add(NetObjectId);
+	ensureMsgf(IsInGameThread(), TEXT("TraceObjectCreated running on parallel thread - this is not yet supported"));
+
+	ObjectTraceData.ObjectsTraced.Add(NetObjectId);
 
 	if (DebugName->DebugNameId == 0U)
 	{
@@ -591,25 +620,33 @@ void FNetTrace::TraceObjectCreated(uint32 GameInstanceId, uint64 NetObjectId, co
 
 void FNetTrace::TraceObjectCreated(uint32 GameInstanceId, uint64 NetObjectId, const FName ObjectName, uint64 TypeIdentifier, uint32 OwnerId)
 {
-	GameInstances.FindOrAdd(GameInstanceId).Objects.Emplace(NetObjectId, { .Name = ObjectName, .DebugName = nullptr, .TypeIdentifier = TypeIdentifier, .ParentNetObjectId = uint64(-1), .OwnerId = OwnerId });
+	FNetTraceInternal::FGameInstance& GameInstance = GameInstances.FindOrAdd(GameInstanceId);
+	FNetTraceInternal::FThreadBuffer& LocalThreadBuffer = FNetTraceInternal::GetThreadBuffer();
+	FNetTraceInternal::FObjectTraceCounter& ObjectTraceData = LocalThreadBuffer.GameInstanceToObjectTraceCounterMap.FindOrAdd(GameInstanceId);
+
+	GameInstance.Objects.Emplace(NetObjectId, { .Name = ObjectName, .DebugName = nullptr, .TypeIdentifier = TypeIdentifier, .ParentNetObjectId = uint64(-1), .OwnerId = OwnerId });
 	
 	if (!GNetTraceRuntimeVerbosity)
 	{
 		return;
 	}
 
-	GameInstances[GameInstanceId].ObjectsTraced.Add(NetObjectId);
+	ObjectTraceData.ObjectsTraced.Add(NetObjectId);
 
 	FNetTraceInternal::Reporter::ReportObjectCreated(GameInstanceId, NetObjectId, FNetTrace::TraceName(ObjectName), TypeIdentifier, OwnerId);
 }
 
 void FNetTrace::TraceObjectDestroyed(uint32 GameInstanceId, uint64 NetObjectId)
 {
-	GameInstances.FindOrAdd(GameInstanceId).Objects.Remove(NetObjectId);
+	FNetTraceInternal::FGameInstance& GameInstance = GameInstances.FindOrAdd(GameInstanceId);
+	FNetTraceInternal::FThreadBuffer& LocalThreadBuffer = FNetTraceInternal::GetThreadBuffer();
+	FNetTraceInternal::FObjectTraceCounter& ObjectTraceData = LocalThreadBuffer.GameInstanceToObjectTraceCounterMap.FindOrAdd(GameInstanceId);
+
+	GameInstance.Objects.Remove(NetObjectId);
 
 	if (GNetTraceRuntimeVerbosity)
 	{
-		GameInstances[GameInstanceId].ObjectsTraced.Remove(NetObjectId);
+		ObjectTraceData.ObjectsTraced.Remove(NetObjectId);
 
 		FNetTraceInternal::Reporter::ReportObjectDestroyed(GameInstanceId, NetObjectId);
 	}
@@ -696,18 +733,18 @@ UE::Net::FNetDebugNameId FNetTrace::TraceName(const TCHAR* Name)
 	}
 
 	// Get Thread buffer
-	FNetTraceInternal::FThreadBuffer* ThreadBufferPtr = FNetTraceInternal::GetThreadBuffer();
+	FNetTraceInternal::FThreadBuffer& LocalThreadBuffer = FNetTraceInternal::GetThreadBuffer();
 
 	// Hash the name using CityHash64
 	const uint64 HashedName = CityHash64((const char*)Name, FCString::Strlen(Name) * sizeof(TCHAR));
-	if (const UE::Net::FNetDebugNameId* FoundNameId = ThreadBufferPtr->DynamicNameHashToNameIdMap.Find(HashedName))
+	if (const UE::Net::FNetDebugNameId* FoundNameId = LocalThreadBuffer.DynamicNameHashToNameIdMap.Find(HashedName))
 	{
 		return *FoundNameId;
 	}
 	else
 	{
 		const UE::Net::FNetDebugNameId NameId = FNetTraceInternal::GetNextNameId();
-		ThreadBufferPtr->DynamicNameHashToNameIdMap.Add(HashedName, NameId);
+		LocalThreadBuffer.DynamicNameHashToNameIdMap.Add(HashedName, NameId);
 
 		FTCHARToUTF8 Converter(Name);
 		FNetTraceInternal::Reporter::ReportAnsiName(NameId, Converter.Length() + 1, (const char*)Converter.Get());
@@ -726,16 +763,16 @@ UE::Net::FNetDebugNameId FNetTrace::TraceName(FName Name)
 	}
 
 	// Get Thread buffer
-	FNetTraceInternal::FThreadBuffer* ThreadBufferPtr = FNetTraceInternal::GetThreadBuffer();
+	FNetTraceInternal::FThreadBuffer& LocalThreadBuffer = FNetTraceInternal::GetThreadBuffer();
 
-	if (const FNetDebugNameId* FoundNameId = ThreadBufferPtr->DynamicFNameToNameIdMap.Find(Name))
+	if (const FNetDebugNameId* FoundNameId = LocalThreadBuffer.DynamicFNameToNameIdMap.Find(Name))
 	{
 		return *FoundNameId;
 	}
 	else
 	{
 		const FNetDebugNameId NameId = FNetTraceInternal::GetNextNameId();
-		ThreadBufferPtr->DynamicFNameToNameIdMap.Add(Name, NameId);
+		LocalThreadBuffer.DynamicFNameToNameIdMap.Add(Name, NameId);
 
 		TCHAR Buffer[FName::StringBufferSize];
 		uint32 NameLen = Name.ToString(Buffer);
@@ -755,8 +792,8 @@ UE::Net::FNetDebugNameId FNetTrace::TraceName(const UE::Net::FNetDebugName* Debu
 
 	if (DebugName->DebugNameId == 0U)
 	{
-		const UE::Net::FNetDebugNameId NameId = TraceName(DebugName->Name);
-		DebugName->DebugNameId = NameId;
+		FNetTraceInternal::FThreadBuffer& LocalThreadBuffer = FNetTraceInternal::GetThreadBuffer();
+		UE::Net::FNetDebugNameId NameId = GetOrCreateDebugNameIdInLocalMap(DebugName, LocalThreadBuffer);
 
 		return NameId;
 	}
@@ -769,19 +806,23 @@ UE::Net::FNetDebugNameId FNetTrace::TraceName(const UE::Net::FNetDebugName* Debu
 void FNetTrace::TraceObjectPolled(uint32 GameInstanceId, uint64 NetObjectId, uint64 Cycles, bool bIsWaste)
 {
 	// not rechecking GNetTraceRuntimeVerbosity here as it's already checked to be == VeryVerbose to have got this far
-
 	FNetTraceInternal::FGameInstance& GameInstance = GameInstances.FindOrAdd(GameInstanceId);
-	TraceExistingObject(GameInstanceId, GameInstance.Objects, GameInstance.ObjectsTraced, NetObjectId);
-
+	FNetTraceInternal::FThreadBuffer& LocalThreadBuffer = FNetTraceInternal::GetThreadBuffer();
+	FNetTraceInternal::FObjectTraceCounter& ObjectTraceData = LocalThreadBuffer.GameInstanceToObjectTraceCounterMap.FindOrAdd(GameInstanceId);
+	
+	TraceExistingObject(GameInstanceId, GameInstance.Objects, ObjectTraceData.ObjectsTraced, NetObjectId);
+	
 	FNetTraceInternal::Reporter::ReportObjectPolled(GameInstanceId, NetObjectId, Cycles, bIsWaste);
 }
 
 void FNetTrace::TraceObjectQuantized(uint32 GameInstanceId, uint64 NetObjectId, uint64 Cycles)
 {
 	// not rechecking GNetTraceRuntimeVerbosity here as it's already checked to be == VeryVerbose to have got this far
-
 	FNetTraceInternal::FGameInstance& GameInstance = GameInstances.FindOrAdd(GameInstanceId);
-	TraceExistingObject(GameInstanceId, GameInstance.Objects, GameInstance.ObjectsTraced, NetObjectId);
+	FNetTraceInternal::FThreadBuffer& LocalThreadBuffer = FNetTraceInternal::GetThreadBuffer();
+	FNetTraceInternal::FObjectTraceCounter& ObjectTraceData = LocalThreadBuffer.GameInstanceToObjectTraceCounterMap.FindOrAdd(GameInstanceId);
+
+	TraceExistingObject(GameInstanceId, GameInstance.Objects, ObjectTraceData.ObjectsTraced, NetObjectId);
 
 	FNetTraceInternal::Reporter::ReportObjectQuantized(GameInstanceId, NetObjectId, Cycles);
 }
@@ -789,16 +830,20 @@ void FNetTrace::TraceObjectQuantized(uint32 GameInstanceId, uint64 NetObjectId,
 void FNetTrace::TraceObjectWritten(uint32 GameInstanceId, uint64 NetObjectId, uint64 Cycles)
 {
 	// not rechecking GNetTraceRuntimeVerbosity here as it's already checked to be == VeryVerbose to have got this far
-
 	FNetTraceInternal::FGameInstance& GameInstance = GameInstances.FindOrAdd(GameInstanceId);
-	TraceExistingObject(GameInstanceId, GameInstance.Objects, GameInstance.ObjectsTraced, NetObjectId);
+	FNetTraceInternal::FThreadBuffer& LocalThreadBuffer = FNetTraceInternal::GetThreadBuffer();
+	FNetTraceInternal::FObjectTraceCounter& ObjectTraceData = LocalThreadBuffer.GameInstanceToObjectTraceCounterMap.FindOrAdd(GameInstanceId);
+
+	TraceExistingObject(GameInstanceId, GameInstance.Objects, ObjectTraceData.ObjectsTraced, NetObjectId);
 
 	FNetTraceInternal::Reporter::ReportObjectWritten(GameInstanceId, NetObjectId, Cycles);
 }
 
 void FNetTrace::TraceSubObject(uint32 GameInstanceId, uint64 ParentNetObjectId, uint64 SubObjectNetObjectId)
 {
-	if (FNetTraceInternal::FNetObject* SubObject = GameInstances.FindOrAdd(GameInstanceId).Objects.Find(SubObjectNetObjectId))
+	FNetTraceInternal::FGameInstance& GameInstance = GameInstances.FindOrAdd(GameInstanceId);
+	FNetTraceInternal::FThreadBuffer& LocalThreadBuffer = FNetTraceInternal::GetThreadBuffer();
+	if (FNetTraceInternal::FNetObject* SubObject = GameInstance.Objects.Find(SubObjectNetObjectId))
 	{
 		SubObject->ParentNetObjectId = ParentNetObjectId;
 	}
diff --git a/Engine/Source/Runtime/Net/Core/Private/Net/Core/Trace/Reporters/NetTraceReporter.cpp b/Engine/Source/Runtime/Net/Core/Private/Net/Core/Trace/Reporters/NetTraceReporter.cpp
--- a/Engine/Source/Runtime/Net/Core/Private/Net/Core/Trace/Reporters/NetTraceReporter.cpp
+++ b/Engine/Source/Runtime/Net/Core/Private/Net/Core/Trace/Reporters/NetTraceReporter.cpp
@@ -24,12 +24,12 @@ UE_TRACE_EVENT_BEGIN(NetTrace, InitEvent)
 UE_TRACE_EVENT_END()
 
 // Trace a name
-UE_TRACE_EVENT_BEGIN(NetTrace, NameEvent)
+UE_TRACE_EVENT_BEGIN(NetTrace, NameEvent, NoSync)
 	UE_TRACE_EVENT_FIELD(uint32, NameId)
 	UE_TRACE_EVENT_FIELD(UE::Trace::AnsiString, Name)
 UE_TRACE_EVENT_END()
 
-UE_TRACE_EVENT_BEGIN(NetTrace, ObjectCreatedEvent)	
+UE_TRACE_EVENT_BEGIN(NetTrace, ObjectCreatedEvent)
 	UE_TRACE_EVENT_FIELD(uint64, TypeId)
 	UE_TRACE_EVENT_FIELD(uint64, ObjectId)
 	UE_TRACE_EVENT_FIELD(uint32, OwnerId)
@@ -37,7 +37,7 @@ UE_TRACE_EVENT_BEGIN(NetTrace, ObjectCreatedEvent)
 	UE_TRACE_EVENT_FIELD(uint8, GameInstanceId)
 UE_TRACE_EVENT_END()
 
-UE_TRACE_EVENT_BEGIN(NetTrace, ObjectExistsEvent)	
+UE_TRACE_EVENT_BEGIN(NetTrace, ObjectExistsEvent, NoSync)	
 	UE_TRACE_EVENT_FIELD(uint64, TypeId)
 	UE_TRACE_EVENT_FIELD(uint64, ObjectId)
 	UE_TRACE_EVENT_FIELD(uint32, OwnerId)
@@ -129,7 +129,7 @@ UE_TRACE_EVENT_BEGIN(NetTrace, PacketDroppedEvent)
 	UE_TRACE_EVENT_FIELD(uint8, PacketType)
 UE_TRACE_EVENT_END()
 
-UE_TRACE_EVENT_BEGIN(NetTrace, ObjectPolledEvent)
+UE_TRACE_EVENT_BEGIN(NetTrace, ObjectPolledEvent, NoSync)
 	UE_TRACE_EVENT_FIELD(uint64, ObjectId)
 	UE_TRACE_EVENT_FIELD(uint64, Cycles)
 	UE_TRACE_EVENT_FIELD(bool, bIsWaste)
@@ -148,7 +148,7 @@ UE_TRACE_EVENT_BEGIN(NetTrace, ObjectWrittenEvent)
 	UE_TRACE_EVENT_FIELD(uint8, GameInstanceId)
 UE_TRACE_EVENT_END()
 
-UE_TRACE_EVENT_BEGIN(NetTrace, SubObjectEvent)
+UE_TRACE_EVENT_BEGIN(NetTrace, SubObjectEvent, NoSync)
 	UE_TRACE_EVENT_FIELD(uint64, ParentObjectId)
 	UE_TRACE_EVENT_FIELD(uint64, SubObjectId)
 	UE_TRACE_EVENT_FIELD(uint8, GameInstanceId)
