Move Iris out of Experimental

diff --git a/Engine/Source/Runtime/Net/Iris/Private/Iris/DataStream/DataStreamDefinitions.cpp b/Engine/Source/Runtime/Net/Iris/Private/Iris/DataStream/DataStreamDefinitions.cpp
new file mode 100644
--- /dev/null
+++ b/Engine/Source/Runtime/Net/Iris/Private/Iris/DataStream/DataStreamDefinitions.cpp
@@ -0,0 +1,62 @@
+// Copyright Epic Games, Inc. All Rights Reserved.
+
+#include "Iris/DataStream/DataStreamDefinitions.h"
+#include "Iris/DataStream/DataStream.h"
+#include "Iris/Core/IrisLog.h"
+#include "UObject/UObjectGlobals.h"
+
+#include UE_INLINE_GENERATED_CPP_BY_NAME(DataStreamDefinitions)
+
+UDataStreamDefinitions::UDataStreamDefinitions()
+: bFixupComplete(false)
+{
+}
+
+void UDataStreamDefinitions::FixupDefinitions()
+{
+	if (bFixupComplete)
+	{
+		return;
+	}
+
+	int32 CurrentStreamIndex = 0;
+	for (FDataStreamDefinition& Definition : DataStreamDefinitions)
+	{
+		UE_CLOG(DataStreamDefinitions.ContainsByPredicate([Name = Definition.DataStreamName, &Definition](const FDataStreamDefinition& ExistingDefinition) { return Name == ExistingDefinition.DataStreamName && &Definition != &ExistingDefinition; }), LogIris, Error, TEXT("DataStream name is defined multiple times: %s."), *Definition.DataStreamName.GetPlainNameString());
+		UE_CLOG(!StaticEnum<EDataStreamSendStatus>()->IsValidEnumValue(int8(Definition.DefaultSendStatus)), LogIris, Error, TEXT("Invalid DataStreamSendStatus %u for DataStream %s."), unsigned(Definition.DefaultSendStatus), *Definition.DataStreamName.GetPlainNameString());
+
+		Definition.Class = StaticLoadClass(UDataStream::StaticClass(), nullptr, *Definition.ClassName.ToString(), nullptr, LOAD_Quiet);
+
+		UE_CLOG(Definition.Class == nullptr, LogIris, Error, TEXT("DataStream class could not be loaded: %s"), *Definition.ClassName.GetPlainNameString());
+
+		Definition.StreamIndex = CurrentStreamIndex++;
+	}
+
+	bFixupComplete = true;
+}
+
+int32 UDataStreamDefinitions::GetStreamIndex(const FDataStreamDefinition& Definition)
+{
+	return Definition.StreamIndex;
+}
+
+const FDataStreamDefinition* UDataStreamDefinitions::FindDefinition(const FName Name) const
+{
+	return DataStreamDefinitions.FindByPredicate([Name](const FDataStreamDefinition& Definition) { return Name == Definition.DataStreamName; });
+}
+
+const FDataStreamDefinition* UDataStreamDefinitions::FindDefinition(int32 StreamIndex) const
+{
+	return DataStreamDefinitions.FindByPredicate([StreamIndex](const FDataStreamDefinition& Definition) { return StreamIndex == Definition.StreamIndex; });
+}
+
+void UDataStreamDefinitions::GetStreamNamesToAutoCreateOrRegister(TArray<FName>& OutStreamNames) const
+{
+	for (const FDataStreamDefinition& Definition : DataStreamDefinitions)
+	{
+		if (Definition.bAutoCreate || Definition.bDynamicCreate)
+		{
+			OutStreamNames.Add(Definition.DataStreamName);
+		}
+	}
+}
diff --git a/Engine/Source/Runtime/Net/Iris/Private/Iris/IrisConfig.cpp b/Engine/Source/Runtime/Net/Iris/Private/Iris/IrisConfig.cpp
new file mode 100644
--- /dev/null
+++ b/Engine/Source/Runtime/Net/Iris/Private/Iris/IrisConfig.cpp
@@ -0,0 +1,42 @@
+// Copyright Epic Games, Inc. All Rights Reserved.
+
+#include "Iris/IrisConfig.h"
+#include "Iris/Core/IrisCsv.h"
+#include "HAL/IConsoleManager.h"
+#include "Misc/CommandLine.h"
+#include "Misc/Parse.h"
+
+namespace UE::Net
+{
+
+static int32 CVarUseIrisReplication = 0;
+static FAutoConsoleVariableRef CVarUseIrisReplicationRef(TEXT("net.Iris.UseIrisReplication"), CVarUseIrisReplication, TEXT("Enables Iris replication system. 0 will fallback to legacy replicationsystem."), ECVF_Default );
+
+bool ShouldUseIrisReplication()
+{
+	return CVarUseIrisReplication > 0;
+}
+
+void SetUseIrisReplication(bool EnableIrisReplication)
+{
+	CVarUseIrisReplication = EnableIrisReplication ? 1 : 0;
+}
+
+EReplicationSystem GetUseIrisReplicationCmdlineValue()
+{
+	int32 UseIrisReplication=0;
+	if (FParse::Value(FCommandLine::Get(), TEXT("UseIrisReplication="), UseIrisReplication))
+	{
+		// Try to force the requested system if the cmdline is present
+		return UseIrisReplication > 0 ? EReplicationSystem::Iris : EReplicationSystem::Generic;
+	}
+
+	// Use the default engine value
+	return EReplicationSystem::Default;
+}
+
+}
+
+
+// Enable Iris category by default on servers
+CSV_DEFINE_CATEGORY(Iris, WITH_SERVER_CODE);
diff --git a/Engine/Source/Runtime/Net/Iris/Private/Iris/ReplicationSystem/Filtering/FilterOutNetObjectFilter.cpp b/Engine/Source/Runtime/Net/Iris/Private/Iris/ReplicationSystem/Filtering/FilterOutNetObjectFilter.cpp
new file mode 100644
--- /dev/null
+++ b/Engine/Source/Runtime/Net/Iris/Private/Iris/ReplicationSystem/Filtering/FilterOutNetObjectFilter.cpp
@@ -0,0 +1,24 @@
+// Copyright Epic Games, Inc. All Rights Reserved.
+
+#include "Iris/ReplicationSystem/Filtering/FilterOutNetObjectFilter.h"
+
+#include UE_INLINE_GENERATED_CPP_BY_NAME(FilterOutNetObjectFilter)
+
+void UFilterOutNetObjectFilter::OnInit(const FNetObjectFilterInitParams& Params)
+{
+}
+
+bool UFilterOutNetObjectFilter::AddObject(uint32 ObjectIndex, FNetObjectFilterAddObjectParams& Params)
+{
+	return true;
+}
+
+void UFilterOutNetObjectFilter::RemoveObject(uint32 ObjectIndex, const FNetObjectFilteringInfo& Info)
+{
+}
+
+void UFilterOutNetObjectFilter::Filter(FNetObjectFilteringParams& Params)
+{
+	// Filter out everything
+	Params.OutAllowedObjects.ClearAllBits();
+}
diff --git a/Engine/Source/Runtime/Net/Iris/Private/Iris/ReplicationSystem/Filtering/NopNetObjectFilter.cpp b/Engine/Source/Runtime/Net/Iris/Private/Iris/ReplicationSystem/Filtering/NopNetObjectFilter.cpp
new file mode 100644
--- /dev/null
+++ b/Engine/Source/Runtime/Net/Iris/Private/Iris/ReplicationSystem/Filtering/NopNetObjectFilter.cpp
@@ -0,0 +1,24 @@
+// Copyright Epic Games, Inc. All Rights Reserved.
+
+#include "Iris/ReplicationSystem/Filtering/NopNetObjectFilter.h"
+
+#include UE_INLINE_GENERATED_CPP_BY_NAME(NopNetObjectFilter)
+
+void UNopNetObjectFilter::OnInit(const FNetObjectFilterInitParams& Params)
+{
+}
+
+bool UNopNetObjectFilter::AddObject(uint32 ObjectIndex, FNetObjectFilterAddObjectParams& Params)
+{
+	return true;
+}
+
+void UNopNetObjectFilter::RemoveObject(uint32 ObjectIndex, const FNetObjectFilteringInfo& Info)
+{
+}
+
+void UNopNetObjectFilter::Filter(FNetObjectFilteringParams& Params)
+{
+	// Allow all objects to be replicated, just as if no dynamic filtering was applied.
+	Params.OutAllowedObjects.SetAllBits();
+}
diff --git a/Engine/Source/Runtime/Net/Iris/Private/Iris/ReplicationSystem/ReplicationBridge.cpp b/Engine/Source/Runtime/Net/Iris/Private/Iris/ReplicationSystem/ReplicationBridge.cpp
new file mode 100644
--- /dev/null
+++ b/Engine/Source/Runtime/Net/Iris/Private/Iris/ReplicationSystem/ReplicationBridge.cpp
@@ -0,0 +1,1295 @@
+// Copyright Epic Games, Inc. All Rights Reserved.
+
+#include "Iris/ReplicationSystem/ReplicationBridge.h"
+
+#include "Containers/ArrayView.h"
+
+#include "HAL/IConsoleManager.h"
+
+#include "Iris/Core/IrisCsv.h"
+#include "Iris/Core/IrisLog.h"
+#include "Iris/Core/IrisProfiler.h"
+#include "Iris/IrisConfigInternal.h"
+
+#include "Iris/ReplicationState/ReplicationStateUtil.h"
+
+#include "Iris/ReplicationSystem/ChangeMaskCache.h"
+#include "Iris/ReplicationSystem/LegacyPushModel.h"
+#include "Iris/ReplicationSystem/NetRefHandleManagerTypes.h"
+#include "Iris/ReplicationSystem/ReplicationFragment.h"
+#include "Iris/ReplicationSystem/ReplicationFragmentInternal.h"
+#include "Iris/ReplicationSystem/ReplicationSystem.h"
+#include "Iris/ReplicationSystem/ReplicationSystemInternal.h"
+#include "Iris/ReplicationSystem/ReplicationWriter.h"
+#include "Iris/ReplicationSystem/RepTag.h"
+#include "Iris/ReplicationSystem/WorldLocations.h"
+
+#include "Iris/ReplicationSystem/Prioritization/NetObjectPrioritizer.h"
+
+#include "Iris/Serialization/InternalNetSerializationContext.h"
+#include "Iris/Serialization/NetBitStreamWriter.h"
+#include "Iris/Serialization/ObjectNetSerializer.h"
+
+#include "Net/Core/NetBitArray.h"
+#include "Net/Core/NetHandle/NetHandleManager.h"
+#include "Net/Core/Trace/NetDebugName.h"
+#include "Net/Core/Trace/NetTrace.h"
+
+#include "ReplicationFragmentInternal.h"
+#include "ReplicationOperationsInternal.h"
+
+#include UE_INLINE_GENERATED_CPP_BY_NAME(ReplicationBridge)
+
+#define UE_LOG_REPLICATIONBRIDGE(Category, Format, ...)  UE_LOG(LogIris, Category, TEXT("RepBridge(%u)::") Format, ReplicationSystem->GetId(), ##__VA_ARGS__)
+
+static bool bEnableFlushReliableRPCOnDestroy = true;
+static FAutoConsoleVariableRef CVarEnableFlushReliableRPCOnDestroy(
+	TEXT("net.Iris.EnableFlushReliableRPCOnDestroy"),
+	bEnableFlushReliableRPCOnDestroy,
+	TEXT("When true EEndReplicationFlags::Flush flag will be appended in EndReplication if we have pending unprocessed attachments/RPC:s when destroying a replicated object.")
+);
+
+static bool bAllowDestroyToCancelFlushAndTearOff = false;
+static FAutoConsoleVariableRef CVarAllowDestroyToCancelFlushAndTearOff(
+	TEXT("net.Iris.AllowDestroyToCancelFlushAndTearOff"),
+	bAllowDestroyToCancelFlushAndTearOff,
+	TEXT("When true issuing a EndReplication on an object that is already Tear-off or pending endreplication will cancel destroy/flush and destroy the replicated objects.")
+);
+
+static bool bAlwaysDestroyDynamicSubObjectInstancesOnDetachFromStaticRoot = true;
+static FAutoConsoleVariableRef CVarAlwaysDestroyDynamicSubObjectInstancesOnDetachFromStaticRoot(
+	TEXT("net.Iris.AlwaysDestroyDynamicSubObjectInstancesOnDetachFromStaticRoot"),
+	bAlwaysDestroyDynamicSubObjectInstancesOnDetachFromStaticRoot,
+	TEXT("When true, We will always destroy instance for dynamic subobjects during EndReplication of a static rootobject.")
+);
+
+/**
+ * ReplicationBridge Implementation
+ */
+UReplicationBridge::UReplicationBridge()
+: ReplicationSystem(nullptr)
+, ReplicationProtocolManager(nullptr)
+, ReplicationStateDescriptorRegistry(nullptr)
+, NetRefHandleManager(nullptr)
+, DestructionInfoProtocol(nullptr)
+{
+}
+
+void UReplicationBridge::PreReceiveUpdate()
+{
+	check(bInReceiveUpdate == false);
+	bInReceiveUpdate = true;
+}
+
+void UReplicationBridge::PostReceiveUpdate()
+{
+	check(bInReceiveUpdate == true);
+	bInReceiveUpdate = false;
+
+	// Now process all StopReplication calls done while inside ReceiveUpdate
+	for (const auto& It : HandlesToStopReplicating)
+	{
+		StopReplicatingNetRefHandle(It.Key, It.Value);
+	}
+	HandlesToStopReplicating.Reset();
+	
+	OnPostReceiveUpdate();
+}
+
+bool UReplicationBridge::WriteNetRefHandleCreationInfo(FReplicationBridgeSerializationContext& Context, FNetRefHandle Handle)
+{
+	return true;
+}
+
+bool UReplicationBridge::CacheNetRefHandleCreationInfo(FNetRefHandle Handle)
+{
+	return false;
+}
+
+FReplicationBridgeCreateNetRefHandleResult UReplicationBridge::CreateNetRefHandleFromRemote(FNetRefHandle RootObjectOfSubObject, FNetRefHandle WantedNetHandle, FReplicationBridgeSerializationContext& Context)
+{
+	return FReplicationBridgeCreateNetRefHandleResult();
+};
+
+void UReplicationBridge::SubObjectCreatedFromReplication(UE::Net::Private::FInternalNetRefIndex RootObjectIndex, FNetRefHandle SubObjectRefHandle)
+{
+}
+
+void UReplicationBridge::DetachInstanceFromRemote(FNetRefHandle Handle, EReplicationBridgeDestroyInstanceReason DestroyReason, EReplicationBridgeDestroyInstanceFlags DestroyFlags,  UE::Net::FNetObjectFactoryId NetFactoryId)
+{
+}
+
+void UReplicationBridge::DetachInstance(FNetRefHandle Handle)
+{
+}
+
+void UReplicationBridge::PruneStaleObjects()
+{
+}
+
+void UReplicationBridge::GetInitialDependencies(FNetRefHandle Handle, FNetDependencyInfoArray& OutDependencies) const
+{
+	return;
+}
+
+void UReplicationBridge::CallGetInitialDependencies(FNetRefHandle Handle, FNetDependencyInfoArray& OutDependencies) const
+{
+	using namespace UE::Net::Private;
+
+	// if the Handle is static, the initial dependency is the handle itself
+	if (Handle.IsStatic())
+	{
+		OutDependencies.Emplace(FObjectReferenceCache::MakeNetObjectReference(Handle));
+		return;
+	}
+	else
+	{
+		GetInitialDependencies(Handle, OutDependencies);
+	}
+}
+
+void UReplicationBridge::DetachSubObjectInstancesFromRemote(FNetRefHandle OwnerHandle, EReplicationBridgeDestroyInstanceReason DestroyReason, EReplicationBridgeDestroyInstanceFlags DestroyFlags)
+{
+	using namespace UE::Net::Private;
+
+	FInternalNetRefIndex OwnerInternalIndex = NetRefHandleManager->GetInternalIndex(OwnerHandle);
+	if (OwnerInternalIndex)
+	{
+		for (const FInternalNetRefIndex SubObjectInternalIndex : NetRefHandleManager->GetChildSubObjects(OwnerInternalIndex))
+		{
+			FNetRefHandleManager::FReplicatedObjectData& SubObjectData = NetRefHandleManager->GetReplicatedObjectDataNoCheck(SubObjectInternalIndex);
+			const FNetRefHandle SubObjectHandle = SubObjectData.RefHandle;
+			SubObjectData.bTearOff = (DestroyReason == EReplicationBridgeDestroyInstanceReason::TearOff);
+			SubObjectData.bPendingEndReplication = 1U;
+
+			EReplicationBridgeDestroyInstanceFlags SubObjectDestroyFlags = DestroyFlags;
+			// The subobject is allowed to be destroyed if both the owner and the subobject allows it.
+			if (!SubObjectData.bAllowDestroyInstanceFromRemote)
+			{
+				EnumRemoveFlags(SubObjectDestroyFlags, EReplicationBridgeDestroyInstanceFlags::AllowDestroyInstanceFromRemote);
+			}
+
+			if (DestroyReason == EReplicationBridgeDestroyInstanceReason::DoNotDestroy && SubObjectData.bAllowDestroyInstanceFromRemote)
+			{
+				// When ending replication of static objects without destroying them, we should always destroy instances spawned from replication
+				// as they will be recreated when the static object is scoped again.
+				if (bAlwaysDestroyDynamicSubObjectInstancesOnDetachFromStaticRoot)
+				{
+					EnumAddFlags(SubObjectDestroyFlags, EReplicationBridgeDestroyInstanceFlags::AllowDestroyInstanceFromRemote);
+					CallDetachInstanceFromRemote(SubObjectHandle, EReplicationBridgeDestroyInstanceReason::Destroy, SubObjectDestroyFlags, SubObjectData.NetFactoryId);
+					continue;
+				}
+				else
+				{
+					UE_LOG(LogIris, Warning, TEXT("Detaching Subobject %s with reason DoNotDestroy even though it is dynamic"), *SubObjectHandle.ToString());
+				}
+			}
+
+			CallDetachInstanceFromRemote(SubObjectHandle, DestroyReason, SubObjectDestroyFlags, SubObjectData.NetFactoryId);
+		}
+	}
+}
+
+void UReplicationBridge::DestroyNetObjectFromRemote(FNetRefHandle Handle, EReplicationBridgeDestroyInstanceReason DestroyReason, EReplicationBridgeDestroyInstanceFlags DestroyFlags)
+{
+	using namespace UE::Net::Private;
+
+	if (Handle.IsValid())
+	{
+		UE_LOG_REPLICATIONBRIDGE(Verbose, TEXT("DestroyNetObjectFromRemote for %s | DestroyReason: %s | DestroyFlags: %s "), *NetRefHandleManager->PrintObjectFromNetRefHandle(Handle), LexToString(DestroyReason), LexToString(DestroyFlags) );
+
+		FInternalNetRefIndex OwnerInternalIndex = NetRefHandleManager->GetInternalIndex(Handle);
+		FNetRefHandleManager::FReplicatedObjectData& ObjectData = NetRefHandleManager->GetReplicatedObjectDataNoCheck(OwnerInternalIndex);
+		ObjectData.bTearOff = (DestroyReason == EReplicationBridgeDestroyInstanceReason::TearOff);
+		ObjectData.bPendingEndReplication = 1U;
+
+		// if a subobject owner is to be destroyed we want to detach all subobjects before doing so to ensure we execute expected callbacks
+		// We keep tracking them internally
+		DetachSubObjectInstancesFromRemote(Handle, DestroyReason, DestroyFlags);
+		
+		// Allow derived bridges to cleanup any instance info they have stored
+		CallDetachInstanceFromRemote(Handle, DestroyReason, DestroyFlags, ObjectData.NetFactoryId);
+
+		// Detach instance protocol
+		InternalDetachInstanceFromNetRefHandle(Handle);
+	
+		// Destroy the NetRefHandle
+		InternalDestroyNetObject(Handle);
+	}
+}
+
+void UReplicationBridge::CallPreSendUpdate(float DeltaSeconds)
+{
+	// Tear-off all handles pending tear-off
+	TearOffHandlesPendingTearOff();
+
+	PreSendUpdate();
+}
+
+void UReplicationBridge::CallPreSendUpdateSingleHandle(FNetRefHandle Handle)
+{
+	PreSendUpdateSingleHandle(Handle);
+}
+
+void UReplicationBridge::CallUpdateInstancesWorldLocation()
+{
+	UpdateInstancesWorldLocation();
+}
+
+void UReplicationBridge::CallDetachInstance(FNetRefHandle Handle)
+{
+	DetachInstance(Handle);
+}
+
+void UReplicationBridge::CallDetachInstanceFromRemote(FNetRefHandle Handle, EReplicationBridgeDestroyInstanceReason DestroyReason, EReplicationBridgeDestroyInstanceFlags DestroyFlags, UE::Net::FNetObjectFactoryId NetFactoryId)
+{
+	DetachInstanceFromRemote(Handle, DestroyReason, DestroyFlags, NetFactoryId);
+}
+
+void UReplicationBridge::CallPruneStaleObjects()
+{
+	PruneStaleObjects();
+}
+
+bool UReplicationBridge::CallCacheNetRefHandleCreationInfo(FNetRefHandle Handle)
+{
+	return CacheNetRefHandleCreationInfo(Handle);
+}
+
+bool UReplicationBridge::CallWriteNetRefHandleCreationInfo(FReplicationBridgeSerializationContext& Context, FNetRefHandle Handle)
+{
+	using namespace UE::Net;
+
+	check(!Context.bIsDestructionInfo);
+
+	return WriteNetRefHandleCreationInfo(Context, Handle);
+}
+
+bool UReplicationBridge::CallWriteNetRefHandleDestructionInfo(FReplicationBridgeSerializationContext& Context, FNetRefHandle Handle)
+{
+	using namespace UE::Net;
+
+	check(Context.bIsDestructionInfo);
+	UE_NET_TRACE_SCOPE(DestructionInfo, *Context.SerializationContext.GetBitStreamWriter(), Context.SerializationContext.GetTraceCollector(), ENetTraceVerbosity::Trace);
+
+	const FDestructionInfo* Info = StaticObjectsPendingDestroy.Find(Handle);
+	if (ensure(Info))
+	{
+		UE_LOG(LogIris, VeryVerbose, TEXT("CallWriteNetRefHandleDestructionInfo on %s | %s | NetFactoryId: %u"), *PrintObjectFromNetRefHandle(Handle), *Info->StaticRef.ToString(), Info->NetFactoryId);
+		// Write destruction info
+		WriteFullNetObjectReference(Context.SerializationContext, Info->StaticRef);
+		Context.SerializationContext.GetBitStreamWriter()->WriteBits(Info->NetFactoryId, FNetObjectFactoryRegistry::GetMaxBits());
+	}
+	else
+	{
+		Context.SerializationContext.SetError(TEXT("DestructionInfoNotFound"), false);
+		UE_LOG_REPLICATIONBRIDGE(Error, TEXT("Failed to write destructionInfo for %s"), *Handle.ToString());
+	}
+
+	return !Context.SerializationContext.HasErrorOrOverflow();
+}
+
+FReplicationBridgeCreateNetRefHandleResult UReplicationBridge::CallCreateNetRefHandleFromRemote(FNetRefHandle RootObjectOfSubObject, FNetRefHandle WantedNetHandle, FReplicationBridgeSerializationContext& Context)
+{
+	check(!Context.bIsDestructionInfo);
+
+	FReplicationBridgeCreateNetRefHandleResult CreateResult = CreateNetRefHandleFromRemote(RootObjectOfSubObject, WantedNetHandle, Context);
+
+	// Track subobjects on clients
+	if (CreateResult.NetRefHandle.IsValid() && RootObjectOfSubObject.IsValid())
+	{
+		NetRefHandleManager->AddSubObject(RootObjectOfSubObject, CreateResult.NetRefHandle);
+	}
+
+	return CreateResult;
+}
+
+bool UReplicationBridge::IsAllowedToDestroyInstance(const UObject* Instance) const
+{
+	return true;
+}
+
+void UReplicationBridge::ReadAndExecuteDestructionInfoFromRemote(FReplicationBridgeSerializationContext& Context)
+{
+	using namespace UE::Net;
+	using namespace UE::Net::Private;
+
+	check(Context.bIsDestructionInfo);
+
+	UE_NET_TRACE_SCOPE(DestructionInfo, *Context.SerializationContext.GetBitStreamReader(), Context.SerializationContext.GetTraceCollector(), ENetTraceVerbosity::Trace);
+
+	// Destroy instance here, or defer to later?
+	FNetObjectReference ReferenceToDestroy;
+	ReadFullNetObjectReference(Context.SerializationContext, ReferenceToDestroy);
+
+	const FNetObjectFactoryId NetFactoryId = IntCastChecked<FNetObjectFactoryId>(Context.SerializationContext.GetBitStreamReader()->ReadBits(FNetObjectFactoryRegistry::GetMaxBits()));
+
+	// Destroy the reference
+	// Resolve the reference in order to be able to destroy it
+	if (const UObject* Instance = ObjectReferenceCache->ResolveObjectReference(ReferenceToDestroy, Context.SerializationContext.GetInternalContext()->ResolveContext))
+	{
+		const FInternalNetRefIndex InternalReplicationIndex = NetRefHandleManager->GetInternalIndex(ReferenceToDestroy.GetRefHandle());
+		if (InternalReplicationIndex != FNetRefHandleManager::InvalidInternalIndex)
+		{
+			NetRefHandleManager->GetReplicatedObjectDataNoCheck(InternalReplicationIndex).bPendingEndReplication = 1U;
+		}
+
+		UE_LOG(LogIris, VeryVerbose, TEXT("ReadAndExecuteDestructionInfoFromRemote on %s | %s | NetFactoryId: %u"), *ReferenceToDestroy.ToString(), *ReferenceToDestroy.GetRefHandle().ToString(), NetFactoryId);
+
+		constexpr EReplicationBridgeDestroyInstanceReason DestroyReason = EReplicationBridgeDestroyInstanceReason::Destroy;
+		const EReplicationBridgeDestroyInstanceFlags DestroyFlags = IsAllowedToDestroyInstance(Instance) ? EReplicationBridgeDestroyInstanceFlags::AllowDestroyInstanceFromRemote : EReplicationBridgeDestroyInstanceFlags::None;
+
+		// if a subobject owner is to be destroyed we want to detach all subobjects before doing so to ensure we execute expected callbacks
+		// We keep tracking them internally
+		DetachSubObjectInstancesFromRemote(ReferenceToDestroy.GetRefHandle(), DestroyReason, DestroyFlags);
+
+		CallDetachInstanceFromRemote(ReferenceToDestroy.GetRefHandle(), DestroyReason, DestroyFlags, NetFactoryId);
+	}
+	else
+	{
+		//$TODO: Need to validate that it's safe if the resolve returns false
+	}
+}
+
+void UReplicationBridge::CallSubObjectCreatedFromReplication(UE::Net::Private::FInternalNetRefIndex RootObjectIndex, FNetRefHandle SubObjectCreated)
+{
+	SubObjectCreatedFromReplication(RootObjectIndex, SubObjectCreated);
+}
+
+void UReplicationBridge::CallPostApplyInitialState(UE::Net::Private::FInternalNetRefIndex InternalObjectIndex)
+{
+	PostApplyInitialState(InternalObjectIndex);
+}
+
+UReplicationBridge::~UReplicationBridge()
+{
+}
+
+void UReplicationBridge::Initialize(UReplicationSystem* InReplicationSystem)
+{
+	using namespace UE::Net;
+	using namespace UE::Net::Private;
+
+	Private::FReplicationSystemInternal* ReplicationSystemInternal = InReplicationSystem->GetReplicationSystemInternal();
+
+	ReplicationSystem = InReplicationSystem;
+	ReplicationProtocolManager = &ReplicationSystemInternal->GetReplicationProtocolManager();
+	ReplicationStateDescriptorRegistry = &ReplicationSystemInternal->GetReplicationStateDescriptorRegistry();
+	NetRefHandleManager = &ReplicationSystemInternal->GetNetRefHandleManager();
+	ObjectReferenceCache = &ReplicationSystemInternal->GetObjectReferenceCache();
+	Groups = &ReplicationSystemInternal->GetGroups();
+
+	// Create destruction info protocol
+	{
+		const FReplicationFragments RegisteredFragments;
+		const FReplicationProtocolIdentifier ProtocolId = FReplicationProtocolManager::CalculateProtocolIdentifier(RegisteredFragments);
+		const FCreateReplicationProtocolParameters CreateProtocolParams
+		{
+			.bValidateProtocolId = false,
+			.bHasTemplateKey = false, // Destruction info protocol doesn't need a TemplateKey
+			.TypeStatsIndex =  GetReplicationSystem()->GetReplicationSystemInternal()->GetNetTypeStats().GetOrCreateTypeStats(FName("DestructionInfo")),
+		};
+		DestructionInfoProtocol = ReplicationProtocolManager->CreateReplicationProtocol(ProtocolId, RegisteredFragments, TEXT("InternalDestructionInfo"), CreateProtocolParams);
+		// Explicit refcount
+		DestructionInfoProtocol->AddRef();
+	}
+}
+
+void UReplicationBridge::Deinitialize()
+{
+	using namespace UE::Net;
+	using namespace UE::Net::Private;
+
+	// Detach all replicated instances that have not yet been destroyed as part of shutting down the rest of the game.
+	NetRefHandleManager->GetAssignedInternalIndices().ForAllSetBits([this](uint32 InternalObjectIndex) 
+	{
+		if (InternalObjectIndex == FNetRefHandleManager::InvalidInternalIndex)
+		{
+			return;
+		}
+
+		FNetRefHandleManager::FReplicatedObjectData& ObjectData = NetRefHandleManager->GetReplicatedObjectDataNoCheck(InternalObjectIndex);
+		if (ObjectData.InstanceProtocol)
+		{
+			const FNetRefHandle RefHandle = ObjectData.RefHandle;
+
+			// Clear out, tracking data or not?, currently we only have a single replication-system so it should not be a big issue if we do.
+			// Note: Currently we opted to leave it as is.
+			// -- This only applies to server. If it is a restart of ReplicationSystem, the same actors will be re-registered and global handle will be destroyed later, otherwise it does not matter what we do.
+			// DestroyGlobalNetHandle(InternalObjectIndex);
+			// ClearNetPushIds(InternalObjectIndex);
+
+			// Detach and destroy instance protocol
+			ObjectData.bPendingEndReplication = 1U;
+			InternalDetachInstanceFromNetRefHandle(ObjectData.RefHandle);
+		}
+	});
+
+	// Release destructioninfo protocol.
+	if (DestructionInfoProtocol)
+	{
+		DestructionInfoProtocol->Release();
+		DestructionInfoProtocol = nullptr;
+	}
+
+	ReplicationSystem = nullptr;
+	ReplicationProtocolManager = nullptr;
+	ReplicationStateDescriptorRegistry = nullptr;
+	NetRefHandleManager = nullptr;
+	ObjectReferenceCache = nullptr;
+	Groups = nullptr;
+}
+
+UE::Net::FNetRefHandle UReplicationBridge::InternalCreateNetObject(FNetRefHandle AllocatedHandle, FNetHandle GlobalHandle, const UE::Net::Private::FCreateNetObjectParams& Params)
+{
+	check(AllocatedHandle.IsValid() && AllocatedHandle.IsCompleteHandle());
+
+	FNetRefHandle Handle = NetRefHandleManager->CreateNetObject(AllocatedHandle, GlobalHandle, Params);
+
+	if (Handle.IsValid())
+	{
+		UE_NET_TRACE_NETHANDLE_CREATED(Handle, Params.ReplicationProtocol->DebugName, (uint64)Params.ReplicationProtocol->ProtocolIdentifier, 0/*Local*/);
+	}
+
+	return Handle;
+}
+
+UE::Net::FNetRefHandle UReplicationBridge::InternalCreateNetObjectFromRemote(FNetRefHandle WantedNetHandle, const UE::Net::Private::FCreateNetObjectParams& Params)
+{
+	FNetRefHandle Handle = NetRefHandleManager->CreateNetObjectFromRemote(WantedNetHandle, Params);
+
+	if (Handle.IsValid())
+	{
+		UE_NET_TRACE_NETHANDLE_CREATED(Handle, Params.ReplicationProtocol->DebugName, (uint64)Params.ReplicationProtocol->ProtocolIdentifier, 1/*Remote*/);
+	}
+
+	return Handle;
+}
+
+void UReplicationBridge::InternalAttachInstanceToNetRefHandle(FNetRefHandle RefHandle, bool bBindInstanceProtocol, UE::Net::FReplicationInstanceProtocol* InstanceProtocol, UObject* Instance, FNetHandle NetHandle)
+{
+	using namespace UE::Net;
+	using namespace UE::Net::Private;
+
+	const uint32 ReplicationSystemId = RefHandle.GetReplicationSystemId();
+	const FInternalNetRefIndex InternalReplicationIndex = NetRefHandleManager->GetInternalIndex(RefHandle);
+
+	NetRefHandleManager->AttachInstanceProtocol(InternalReplicationIndex, InstanceProtocol, Instance);
+	UE_LOG_REPLICATIONBRIDGE(Verbose, TEXT("InternalAttachInstanceToNetHandle Attached: %s %s to (InternalIndex: %u)"), *Instance->GetName(), *RefHandle.ToString(), InternalReplicationIndex);
+
+	// Bind instance protocol to dirty state tracking
+	if (bBindInstanceProtocol)
+	{
+		FReplicationInstanceOperationsInternal::BindInstanceProtocol(NetHandle, InstanceProtocol, NetRefHandleManager->GetReplicatedObjectDataNoCheck(InternalReplicationIndex).Protocol);
+		ForceNetUpdate(ReplicationSystemId, InternalReplicationIndex);
+	}
+}
+
+void UReplicationBridge::InternalDetachInstanceFromNetRefHandle(FNetRefHandle RefHandle)
+{
+	using namespace UE::Net;
+	using namespace UE::Net::Private;
+
+	const FInternalNetRefIndex InternalReplicationIndex = NetRefHandleManager->GetInternalIndex(RefHandle);
+
+	if (FReplicationInstanceProtocol* InstanceProtocol = const_cast<FReplicationInstanceProtocol*>(NetRefHandleManager->DetachInstanceProtocol(InternalReplicationIndex)))
+	{
+		UE_LOG_REPLICATIONBRIDGE(Verbose, TEXT("InternalDetachInstanceToNetHandle Detached: %s from (InternalIndex: %u)"), *RefHandle.ToString(), InternalReplicationIndex);
+
+		if (EnumHasAnyFlags(InstanceProtocol->InstanceTraits, EReplicationInstanceProtocolTraits::IsBound))
+		{
+			FReplicationInstanceOperationsInternal::UnbindInstanceProtocol(InstanceProtocol, NetRefHandleManager->GetReplicatedObjectDataNoCheck(InternalReplicationIndex).Protocol);
+		}
+		ReplicationProtocolManager->DestroyInstanceProtocol(InstanceProtocol);
+	}
+}
+
+void UReplicationBridge::InternalDestroyNetObject(FNetRefHandle Handle)
+{
+	using namespace UE::Net;
+	using namespace UE::Net::Private;
+
+	if (const FInternalNetRefIndex ObjectInternalIndex = NetRefHandleManager->GetInternalIndex(Handle))
+	{
+		FReplicationSystemInternal* ReplicationSystemInternal = ReplicationSystem->GetReplicationSystemInternal();
+
+		FWorldLocations& WorldLocations = ReplicationSystemInternal->GetWorldLocations();
+		WorldLocations.RemoveObjectInfoCache(ObjectInternalIndex);
+
+		// Remove from handles to stop replicating as we might respawn the same object again in a later packet
+		HandlesToStopReplicating.Remove(Handle);
+
+		FNetRefHandleManager::FReplicatedObjectData& ObjectData = NetRefHandleManager->GetReplicatedObjectDataNoCheck(ObjectInternalIndex);
+
+		// At this point we can no longer instantiate objects and must explicitly clear out objects pending create from ReplicationWriter
+		if (ObjectData.bHasCachedCreationInfo && NetRefHandleManager->GetNetObjectRefCount(ObjectInternalIndex) > 0U)
+		{
+			// We need to explicitly notify all ReplicationWriters that we are destroying objects pending tearoff
+			// The handle will automatically be removed from HandlesPendingEndReplication after the next update
+			FReplicationConnections& Connections = ReplicationSystemInternal->GetConnections();
+
+			auto NotifyDestroyedObjectPendingEndReplication = [&Connections, &ObjectInternalIndex](uint32 ConnectionId)
+			{
+				FReplicationConnection* Conn = Connections.GetConnection(ConnectionId);			
+				Conn->ReplicationWriter->NotifyDestroyedObjectPendingEndReplication(ObjectInternalIndex);
+			};
+
+			const FNetBitArray& ValidConnections = Connections.GetValidConnections();
+			ValidConnections.ForAllSetBits(NotifyDestroyedObjectPendingEndReplication);					
+
+			// As we have release cached creation headers before calling InternalDestroyNetObject it is critical that we also clear out the flag.
+			ObjectData.bHasCachedCreationInfo = 0U;
+		}
+	}
+
+	NetRefHandleManager->DestroyNetObject(Handle);
+}
+
+void UReplicationBridge::DestroyLocalNetHandle(FNetRefHandle Handle, EEndReplicationFlags EndReplicationFlags)
+{
+	UE_LOG_REPLICATIONBRIDGE(Verbose, TEXT("DestroyLocalNetHandle for %s | EndReplicationFlags: %s"), *NetRefHandleManager->PrintObjectFromNetRefHandle(Handle), *LexToString(EndReplicationFlags));
+
+	if (EnumHasAnyFlags(EndReplicationFlags, EEndReplicationFlags::DestroyNetHandle | EEndReplicationFlags::ClearNetPushId))
+	{
+		const UE::Net::Private::FInternalNetRefIndex InternalReplicationIndex = NetRefHandleManager->GetInternalIndex(Handle);
+
+		if (EnumHasAnyFlags(EndReplicationFlags, EEndReplicationFlags::DestroyNetHandle))
+		{
+			DestroyGlobalNetHandle(InternalReplicationIndex);
+		}
+
+		if (EnumHasAnyFlags(EndReplicationFlags, EEndReplicationFlags::ClearNetPushId))
+		{
+			ClearNetPushIds(InternalReplicationIndex);
+		}
+	}
+
+	// Detach instance protocol
+	InternalDetachInstanceFromNetRefHandle(Handle);
+
+	// Allow derived bridges to cleanup any instance info they have stored
+	CallDetachInstance(Handle);
+
+	// If the object is in any groups we need to remove it to make sure that we update filtering
+	GetReplicationSystem()->RemoveFromAllGroups(Handle);
+
+	// If we have any attached SubObjects, tag them for destroy as well
+	InternalDestroySubObjects(Handle, EndReplicationFlags);
+
+	// Tell ReplicationSystem to destroy the handle
+	InternalDestroyNetObject(Handle);
+}
+
+void UReplicationBridge::InternalAddSubObject(FNetRefHandle OwnerHandle, FNetRefHandle SubObjectHandle, FNetRefHandle InsertRelativeToSubObjectHandle, ESubObjectInsertionOrder InsertionOrder)
+{
+	using namespace UE::Net::Private;
+
+	EAddSubObjectFlags AddSubObjectFlags = EAddSubObjectFlags::Default;
+
+	switch(InsertionOrder)
+	{
+		case ESubObjectInsertionOrder::None: 
+			break;
+		case ESubObjectInsertionOrder::ReplicateWith:
+			AddSubObjectFlags |= EAddSubObjectFlags::ReplicateWithSubObject;
+			break;
+		case ESubObjectInsertionOrder::InsertAtStart:
+			AddSubObjectFlags |= EAddSubObjectFlags::InsertAtStart;
+			break;
+		default:
+			checkf(false, TEXT("Missing implementation of ESubObjectInsertionOrder enum"));
+			break;
+	}
+
+	if (NetRefHandleManager->AddSubObject(OwnerHandle, SubObjectHandle, InsertRelativeToSubObjectHandle, AddSubObjectFlags))
+	{
+		// If the subobject is new we need to update it immediately to pick it up for replication with its new parent
+		ForceNetUpdate(ReplicationSystem->GetId(), NetRefHandleManager->GetInternalIndex(SubObjectHandle));
+
+		// We set the priority of subobjects to be static as they will be prioritized with owner
+		ReplicationSystem->SetStaticPriority(SubObjectHandle, 1.0f);
+	}
+}
+
+void UReplicationBridge::InternalDestroySubObjects(FNetRefHandle OwnerHandle, EEndReplicationFlags Flags)
+{
+	using namespace UE::Net::Private;
+
+	// Destroy SubObjects
+	FInternalNetRefIndex OwnerInternalIndex = NetRefHandleManager->GetInternalIndex(OwnerHandle);
+	if (OwnerInternalIndex)
+	{
+		for (FInternalNetRefIndex SubObjectInternalIndex : NetRefHandleManager->GetChildSubObjects(OwnerInternalIndex))
+		{
+			FNetRefHandleManager::FReplicatedObjectData& SubObjectData = NetRefHandleManager->GetReplicatedObjectDataNoCheck(SubObjectInternalIndex);
+
+			const FNetRefHandle SubObjectHandle = SubObjectData.RefHandle;
+			const bool bDestroySubObjectWithOwner = SubObjectData.bDestroySubObjectWithOwner;
+				
+			// Tag subobject for destroy. The check against the scope is needed since the subobjects array might contain subobjects already pending destroy.
+			if (bDestroySubObjectWithOwner && NetRefHandleManager->IsScopableIndex(SubObjectInternalIndex))
+			{
+				SubObjectData.bPendingEndReplication = 1U;
+				UE_LOG_REPLICATIONBRIDGE(Verbose, TEXT("InternalDestroySubObjects %s"), *NetRefHandleManager->PrintObjectFromNetRefHandle(SubObjectHandle));
+				DestroyLocalNetHandle(SubObjectHandle, Flags);
+			}
+		}
+	}
+}
+
+void UReplicationBridge::StopReplicatingNetRefHandle(FNetRefHandle Handle, EEndReplicationFlags EndReplicationFlags)
+{
+	using namespace UE::Net::Private;
+
+	if (!IsReplicatedHandle(Handle))
+	{
+		return;
+	}
+
+	const FInternalNetRefIndex InternalReplicationIndex = NetRefHandleManager->GetInternalIndex(Handle);
+	if (NetRefHandleManager->IsLocal(InternalReplicationIndex))
+	{
+		FNetRefHandleManager::FReplicatedObjectData& ObjectData = NetRefHandleManager->GetReplicatedObjectDataNoCheck(InternalReplicationIndex);
+
+		if (ObjectData.bPendingEndReplication && !bAllowDestroyToCancelFlushAndTearOff)
+		{
+			UE_LOG_REPLICATIONBRIDGE(Verbose, TEXT("Ignoring EndReplication called on object already PendingEndReplication %s."), *NetRefHandleManager->PrintObjectFromNetRefHandle(Handle));
+			return;
+		}
+
+		if (EnumHasAnyFlags(EndReplicationFlags, EEndReplicationFlags::TearOff))
+		{
+			// Add handle to list of objects pending EndReplication indicate that it should be destroyed during next update
+			// We need to do this to cover the case where the torn off object not yet has been added to the scope.
+			AddPendingEndReplication(Handle, EndReplicationFlags);
+
+			// We do however copy the final state data and mark object to stop propagating state changes
+			InternalTearOff(Handle);
+
+			// Detach instance as we must assume that we should not access the object after this call.
+			ObjectData.bPendingEndReplication = 1U;
+			InternalDetachInstanceFromNetRefHandle(Handle);
+		}
+		else 
+		{
+			// New objects, destroyed during the same frame with posted attachments(RPC):s needs to request a flush to ensure that they get a scope update
+			const UE::Net::Private::FNetBlobManager& NetBlobManager = GetReplicationSystem()->GetReplicationSystemInternal()->GetNetBlobManager();
+			const bool bAllowAutoFlushOfUnProcessedReliableRPCs = bEnableFlushReliableRPCOnDestroy && ObjectData.bNeedsFullCopyAndQuantize;
+			if (bAllowAutoFlushOfUnProcessedReliableRPCs && NetBlobManager.HasUnprocessedReliableAttachments(InternalReplicationIndex))
+			{
+				EnumAddFlags(EndReplicationFlags, EEndReplicationFlags::Flush);
+			}
+
+			if (EnumHasAnyFlags(EndReplicationFlags, EEndReplicationFlags::Flush))
+			{
+				// Defer destroy until after scope update to allow create/destroy on same frame
+				AddPendingEndReplication(Handle, EndReplicationFlags);
+			
+				// Capture final state
+				InternalFlushStateData(Handle);
+
+				// Detach instance as we must assume that we should not access the object after this call.
+				ObjectData.bPendingEndReplication = 1U;
+				InternalDetachInstanceFromNetRefHandle(Handle);
+			}
+			else
+			{
+				ObjectData.bPendingEndReplication = 1U;
+				DestroyLocalNetHandle(Handle, EndReplicationFlags);	
+			}	
+		}
+	}
+	else
+	{
+
+		// While we are inside ReceiveUpdate, queue stop replication requests instead of immediately stopping replication
+		// This allows us the apply any received updates before we cut off this object
+		if (IsInReceiveUpdate())
+		{
+			UE_LOG(LogIris, Verbose, TEXT("Delayed request to StopReplicating %s (flags: %s) because it was called while inside ReceiveUpdate"), *NetRefHandleManager->PrintObjectFromNetRefHandle(Handle), *LexToString(EndReplicationFlags));
+		
+			// Detect if we have diverging EndReplicationFlags for the same netobject
+	#if DO_ENSURE
+			{
+				EEndReplicationFlags* PreviousFlags = HandlesToStopReplicating.Find(Handle);
+				const bool bPreviousFlagsMatch = PreviousFlags ? (*PreviousFlags) == EndReplicationFlags : true;
+				ensureMsgf(bPreviousFlagsMatch, TEXT("Received multiple StopReplicating calls for %s with different EndReplicationFlags: Previous: %s | Newest: %s"),
+					*NetRefHandleManager->PrintObjectFromNetRefHandle(Handle), *LexToString(*PreviousFlags), *LexToString(EndReplicationFlags));
+			}
+	#endif
+
+			HandlesToStopReplicating.Add(Handle, EndReplicationFlags);
+			return;
+		}
+
+		if (InternalReplicationIndex != FNetRefHandleManager::InvalidInternalIndex && EnumHasAnyFlags(EndReplicationFlags, EEndReplicationFlags::SkipPendingEndReplicationValidation))
+		{
+			NetRefHandleManager->GetReplicatedObjectDataNoCheck(InternalReplicationIndex).bPendingEndReplication = 1U;
+		}
+		// If we get a call to end replication on the client, we need to detach the instance as it might be garbage collected
+		InternalDetachInstanceFromNetRefHandle(Handle);
+	}
+}
+
+void UReplicationBridge::RemoveDestructionInfosForGroup(UE::Net::FNetObjectGroupHandle GroupHandle)
+{
+	using namespace UE::Net;
+	using namespace UE::Net::Private;
+
+	FWorldLocations& WorldLocations = GetReplicationSystem()->GetReplicationSystemInternal()->GetWorldLocations();
+
+	if (GroupHandle.IsValid())
+	{
+		UE_LOG_REPLICATIONBRIDGE(Verbose, TEXT("RemoveDestructionInfosForGroup GroupIndex: %u"), GroupHandle.GetGroupIndex());
+
+		FNetObjectGroup* Group = Groups->GetGroup(GroupHandle);
+		check(Group);
+
+		TArray<FNetRefHandle, TInlineAllocator<384>> ObjectsToRemove;	
+		TArray<FInternalNetRefIndex, TInlineAllocator<384>> ObjectIndicesToRemove;	
+		for (uint32 InternalObjectIndex : MakeArrayView(Group->Members))
+		{
+			if (NetRefHandleManager->GetIsDestroyedStartupObject(InternalObjectIndex))
+			{
+				const FNetRefHandleManager::FReplicatedObjectData& ObjectData = NetRefHandleManager->GetReplicatedObjectDataNoCheck(InternalObjectIndex);
+				if (StaticObjectsPendingDestroy.Remove(ObjectData.RefHandle))
+				{
+					ObjectsToRemove.Add(ObjectData.RefHandle);
+					ObjectIndicesToRemove.Add(InternalObjectIndex);
+				}
+			}			
+		}
+
+		for (FNetRefHandle Handle : MakeArrayView(ObjectsToRemove))
+		{
+			NetRefHandleManager->DestroyNetObject(Handle);
+		}
+
+		for (FInternalNetRefIndex InternalReplicationIndex : MakeArrayView(ObjectIndicesToRemove))
+		{
+			WorldLocations.RemoveObjectInfoCache(InternalReplicationIndex);
+		}
+	}
+	else
+	{
+		// We should remove all destruction infos and objects
+		for (const auto& It : StaticObjectsPendingDestroy)
+		{
+			NetRefHandleManager->DestroyNetObject(It.Key);
+		}
+
+		// Remove from WorldLocations
+		for (const auto& It : StaticObjectsPendingDestroy)
+		{
+			WorldLocations.RemoveObjectInfoCache(It.Value.InternalReplicationIndex);
+		}
+
+		StaticObjectsPendingDestroy.Empty();
+	}
+}
+
+void UReplicationBridge::TearOffHandlesPendingTearOff()
+{
+	// Initiate tear off
+	for (FPendingEndReplicationInfo Info : MakeArrayView(HandlesPendingEndReplication))
+	{
+		if (EnumHasAnyFlags(Info.DestroyFlags, EEndReplicationFlags::TearOff))
+		{
+			InternalTearOff(Info.Handle);
+		}
+	}
+}
+
+void UReplicationBridge::UpdateHandlesPendingEndReplication()
+{
+	using namespace UE::Net;
+	using namespace UE::Net::Private;
+
+	TArray<FPendingEndReplicationInfo, TInlineAllocator<32>> ObjectsStillPendingEndReplication;
+	for (FPendingEndReplicationInfo Info : MakeArrayView(HandlesPendingEndReplication))
+	{
+		if (FInternalNetRefIndex ObjectInternalIndex = NetRefHandleManager->GetInternalIndex(Info.Handle))
+		{
+			// Immediate destroy or objects that are no longer are referenced by any connections are destroyed
+			if (NetRefHandleManager->GetNetObjectRefCount(ObjectInternalIndex) == 0U || Info.Immediate == EPendingEndReplicationImmediate::Yes)
+			{
+				FNetRefHandleManager::FReplicatedObjectData& ObjectData = NetRefHandleManager->GetReplicatedObjectDataNoCheck(ObjectInternalIndex);
+				ObjectData.bPendingEndReplication = 1U;
+				DestroyLocalNetHandle(Info.Handle, Info.DestroyFlags);
+			}
+			else
+			{
+				// If the object is still in scope remove it from scope as objects pending EndReplication should not be added to new connections after the first update.
+				if (NetRefHandleManager->IsScopableIndex(ObjectInternalIndex))
+				{
+					// Mark object and subobjects as no longer scopeable, and that we should not propagate changed states
+					NetRefHandleManager->RemoveFromScope(ObjectInternalIndex);
+					for (FInternalNetRefIndex SubObjectIndex : NetRefHandleManager->GetSubObjects(ObjectInternalIndex))
+					{
+						NetRefHandleManager->RemoveFromScope(SubObjectIndex);
+					}
+				}
+			
+				// Keep object in the pending EndReplication-list until the object is no longer referenced by any ReplicationWriter
+				ObjectsStillPendingEndReplication.Add(FPendingEndReplicationInfo(Info.Handle, Info.DestroyFlags, EPendingEndReplicationImmediate::No));
+			}
+		}
+	}
+
+	HandlesPendingEndReplication.Reset();
+	HandlesPendingEndReplication.Insert(ObjectsStillPendingEndReplication.GetData(), ObjectsStillPendingEndReplication.Num(), 0);
+
+	CSV_CUSTOM_STAT(Iris, NumHandlesPendingEndRepliction, (float)HandlesPendingEndReplication.Num(), ECsvCustomStatOp::Set);
+}
+
+void UReplicationBridge::AddPendingEndReplication(FNetRefHandle Handle, EEndReplicationFlags DestroyFlags, EPendingEndReplicationImmediate Immediate)
+{
+	if (ensure(EnumHasAnyFlags(DestroyFlags, EEndReplicationFlags::Flush | EEndReplicationFlags::TearOff)))
+	{
+		if (!HandlesPendingEndReplication.FindByPredicate([&](const FPendingEndReplicationInfo& Info){ return Info.Handle == Handle; }))
+		{
+			HandlesPendingEndReplication.Add(FPendingEndReplicationInfo(Handle, DestroyFlags, Immediate));
+		}
+	}
+}
+
+void UReplicationBridge::InternalFlushStateData(UE::Net::FNetSerializationContext& SerializationContext, UE::Net::Private::FChangeMaskCache& ChangeMaskCache, UE::Net::FNetBitStreamWriter& ChangeMaskWriter, uint32 InternalObjectIndex)
+{
+	using namespace UE::Net;
+	using namespace UE::Net::Private;
+
+	if (InternalObjectIndex == FNetRefHandleManager::InvalidInternalIndex)
+	{
+		return;
+	}
+
+	FNetRefHandleManager::FReplicatedObjectData& ObjectData = NetRefHandleManager->GetReplicatedObjectDataNoCheck(InternalObjectIndex);
+
+	// Copy state data, if object already is torn off there is nothing to do
+	if (ObjectData.bTearOff)
+	{
+		return;
+	}
+
+	UE_LOG_REPLICATIONBRIDGE(Verbose, TEXT("InternalFlushStateData Initiating flush for %s (InternalIndex: %u)"), *ObjectData.RefHandle.ToString(), InternalObjectIndex);
+
+	if (ObjectData.InstanceProtocol)
+	{
+		if (EnumHasAnyFlags(ObjectData.InstanceProtocol->InstanceTraits, EReplicationInstanceProtocolTraits::NeedsPoll | EReplicationInstanceProtocolTraits::NeedsPreSendUpdate))
+		{
+			CallPreSendUpdateSingleHandle(ObjectData.RefHandle);
+		}
+
+		// Cache creation info
+		ObjectData.bHasCachedCreationInfo =  CallCacheNetRefHandleCreationInfo(ObjectData.RefHandle) ? 1U : 0U;
+
+		FReplicationInstanceOperationsInternal::QuantizeObjectStateData(ChangeMaskWriter, ChangeMaskCache, *NetRefHandleManager, SerializationContext, InternalObjectIndex);
+
+		// Clear the quantize flag since it was done directly here.
+		NetRefHandleManager->GetDirtyObjectsToQuantize().ClearBit(InternalObjectIndex);
+	}
+
+	for (FInternalNetRefIndex SubObjectInternalIndex : NetRefHandleManager->GetChildSubObjects(InternalObjectIndex))
+	{
+		InternalFlushStateData(SerializationContext, ChangeMaskCache, ChangeMaskWriter, SubObjectInternalIndex);
+	}
+
+	// $IRIS TODO:  Should we also clear the DirtyTracker flags for this flushed object ?
+}
+
+void UReplicationBridge::InternalFlushStateData(FNetRefHandle Handle)
+{
+	using namespace UE::Net;
+	using namespace UE::Net::Private;
+
+	IRIS_PROFILER_SCOPE(InternalFlushStateData);
+
+	const uint32 InternalObjectIndex = NetRefHandleManager->GetInternalIndex(Handle);
+	if (InternalObjectIndex == FNetRefHandleManager::InvalidInternalIndex)
+	{
+		return;
+	}
+
+	FNetRefHandleManager::FReplicatedObjectData& ObjectData = NetRefHandleManager->GetReplicatedObjectDataNoCheck(InternalObjectIndex);
+
+	FChangeMaskCache ChangeMaskCache;
+	FNetBitStreamWriter ChangeMaskWriter;
+
+	// Setup context
+	FNetSerializationContext SerializationContext;
+	FInternalNetSerializationContext InternalContext(ReplicationSystem);
+	SerializationContext.SetInternalContext(&InternalContext);
+	SerializationContext.SetNetStatsContext(ReplicationSystem->GetReplicationSystemInternal()->GetNetTypeStats().GetNetStatsContext());
+
+	InternalFlushStateData(SerializationContext, ChangeMaskCache, ChangeMaskWriter, InternalObjectIndex);
+
+	// Iterate over connections and propagate dirty changemasks to all connections already scoping this object
+	if (ChangeMaskCache.Indices.Num() > 0)
+	{
+		FReplicationConnections& Connections = ReplicationSystem->GetReplicationSystemInternal()->GetConnections();
+
+		auto&& UpdateDirtyChangeMasks = [&Connections, &ChangeMaskCache](uint32 ConnectionId)
+		{
+			FReplicationConnection* Conn = Connections.GetConnection(ConnectionId);			
+
+			const bool bMarkForTearOff = false;
+			Conn->ReplicationWriter->ForceUpdateDirtyChangeMasks(ChangeMaskCache, FReplicationWriter::FlushFlags_FlushState, bMarkForTearOff);
+		};
+		const FNetBitArray& ValidConnections = Connections.GetValidConnections();
+		ValidConnections.ForAllSetBits(UpdateDirtyChangeMasks);		
+	}
+}
+
+void UReplicationBridge::InternalTearOff(FNetRefHandle Handle)
+{
+	using namespace UE::Net;
+	using namespace UE::Net::Private;
+
+	IRIS_PROFILER_SCOPE(InternalTearOff);
+
+	const uint32 InternalObjectIndex = NetRefHandleManager->GetInternalIndex(Handle);
+	if (InternalObjectIndex == FNetRefHandleManager::InvalidInternalIndex)
+	{
+		return;
+	}
+
+	FNetRefHandleManager::FReplicatedObjectData& ObjectData = NetRefHandleManager->GetReplicatedObjectDataNoCheck(InternalObjectIndex);
+
+	if (ObjectData.bTearOff)
+	{
+		// Already torn off
+		return;
+	}
+
+	// Copy state data and tear off now
+	UE_LOG_REPLICATIONBRIDGE(Verbose, TEXT("TearOff: %s"), *PrintObjectFromNetRefHandle(Handle));
+
+	// Force copy of final state data as we will detach the object after scope update
+	FChangeMaskCache ChangeMaskCache;
+	FNetBitStreamWriter ChangeMaskWriter;
+
+	// Setup context
+	FNetSerializationContext SerializationContext;
+	FInternalNetSerializationContext InternalContext(ReplicationSystem);
+	SerializationContext.SetInternalContext(&InternalContext);
+	SerializationContext.SetNetStatsContext(ReplicationSystem->GetReplicationSystemInternal()->GetNetTypeStats().GetNetStatsContext());
+
+	if (ObjectData.InstanceProtocol)
+	{
+		if (EnumHasAnyFlags(ObjectData.InstanceProtocol->InstanceTraits, EReplicationInstanceProtocolTraits::NeedsPoll | EReplicationInstanceProtocolTraits::NeedsPreSendUpdate))
+		{
+			CallPreSendUpdateSingleHandle(Handle);
+		} 
+
+		// Cache creation info
+		ObjectData.bHasCachedCreationInfo =  CallCacheNetRefHandleCreationInfo(Handle) ? 1U : 0U;
+	}
+
+	if (ObjectData.InstanceProtocol && ObjectData.Protocol->InternalTotalSize > 0U)
+	{
+		FReplicationInstanceOperationsInternal::QuantizeObjectStateData(ChangeMaskWriter, ChangeMaskCache, *NetRefHandleManager, SerializationContext, InternalObjectIndex);
+
+		// Clear the quantize flag since it was done directly here.
+		NetRefHandleManager->GetDirtyObjectsToQuantize().ClearBit(InternalObjectIndex);
+	}
+	else
+	{
+		// Nothing to copy, but we must still propagate the tear-off state.
+		FChangeMaskCache::FCachedInfo& Info = ChangeMaskCache.AddEmptyChangeMaskForObject(InternalObjectIndex);
+		// If we are a subobject we must also mark owner as dirty.
+		const uint32 SubObjectOwnerIndex = ObjectData.SubObjectRootIndex;
+		if (SubObjectOwnerIndex != FNetRefHandleManager::InvalidInternalIndex) 
+		{
+			ChangeMaskCache.AddSubObjectOwnerDirty(SubObjectOwnerIndex);
+		}			
+	}
+
+	// Propagate changes to all connections that we currently have in scope
+	FReplicationConnections& Connections = ReplicationSystem->GetReplicationSystemInternal()->GetConnections();
+
+	// Iterate over connections and propagate dirty changemasks to all connections already scoping this object
+	auto UpdateDirtyChangeMasks = [&Connections, &ChangeMaskCache](uint32 ConnectionId)
+	{
+		FReplicationConnection* Conn = Connections.GetConnection(ConnectionId);
+		const bool bMarkForTearOff = true;
+		Conn->ReplicationWriter->ForceUpdateDirtyChangeMasks(ChangeMaskCache, FReplicationWriter::FlushFlags_None, bMarkForTearOff);
+	};
+	const FNetBitArray& ValidConnections = Connections.GetValidConnections();
+	ValidConnections.ForAllSetBits(UpdateDirtyChangeMasks);		
+
+	// TearOff subobjects as well.
+	for (FInternalNetRefIndex SubObjectInternalIndex : NetRefHandleManager->GetChildSubObjects(InternalObjectIndex))
+	{
+		InternalTearOff(NetRefHandleManager->GetNetRefHandleFromInternalIndex(SubObjectInternalIndex));
+	}	
+
+	// Mark object as being torn-off and that we should no longer propagate state changes
+	ObjectData.bTearOff = 1U;
+	ObjectData.bShouldPropagateChangedStates = 0U;
+}
+
+UE::Net::FNetRefHandle UReplicationBridge::StoreDestructionInfo(FNetRefHandle Handle, const FDestructionParameters& Parameters)
+{
+	using namespace UE::Net;
+	using namespace UE::Net::Private;
+
+	if (!FNetObjectFactoryRegistry::IsValidFactoryId(Parameters.NetFactoryId))
+	{
+		ensureMsgf(false, TEXT("StoreDestructionInfo received invalid NetFactoryId: %u for %s"), Parameters.NetFactoryId, *PrintObjectFromNetRefHandle(Handle));
+		UE_LOG(LogIris, Error, TEXT("StoreDestructionInfo received invalid NetFactoryID: %u for %s"), Parameters.NetFactoryId, *PrintObjectFromNetRefHandle(Handle));
+
+		return FNetRefHandle::GetInvalid();
+	}
+
+	// Not allowed to create destruction infos at the moment
+	if (!CanCreateDestructionInfo())
+	{
+		return FNetRefHandle::GetInvalid();
+	}
+
+	FReplicationProtocolManager* ProtocolManager = GetReplicationProtocolManager();
+
+	const FCreateNetObjectParams CreateNetObjectParams
+	{
+		.NetFactoryId = Parameters.NetFactoryId,
+		.ReplicationProtocol = DestructionInfoProtocol,
+	};
+
+	// Create handle for DestructionInfo to be able to scope destruction infos
+	const FNetObjectGroupHandle LevelGroupHandle = GetLevelGroup(Parameters.Level);
+	FNetRefHandle DestructionInfoHandle = NetRefHandleManager->CreateHandleForDestructionInfo(Handle, CreateNetObjectParams);
+
+	if (LevelGroupHandle.IsValid())
+	{
+		GetReplicationSystem()->AddToGroup(LevelGroupHandle, DestructionInfoHandle);
+	}
+
+	const FInternalNetRefIndex InternalReplicationIndex = NetRefHandleManager->GetInternalIndex(DestructionInfoHandle);
+
+	// We also need to store the actual data we send to destroy static references when they are scoped
+	FDestructionInfo PersistentDestructionInfo
+	{
+		.StaticRef = FObjectReferenceCache::MakeNetObjectReference(Handle),
+		.LevelGroupHandle = LevelGroupHandle,
+		.NetFactoryId = Parameters.NetFactoryId,
+		.InternalReplicationIndex = InternalReplicationIndex
+	};
+
+	StaticObjectsPendingDestroy.Add(DestructionInfoHandle, PersistentDestructionInfo);
+
+	// If we use distance based Prioritization for destruction infos we need to populate the quantized state used for prioritization.
+	const bool bUseDynamicPrioritization = Parameters.bUseDistanceBasedPrioritization;
+	if (bUseDynamicPrioritization)
+	{
+		// Use WorldLocations to feed the location of the destruction info so that it can be prioritized properly.
+		FWorldLocations& WorldLocations = GetReplicationSystem()->GetReplicationSystemInternal()->GetWorldLocations();
+
+		// Check the position lies within world boundaries.
+		ensureMsgf(Parameters.Location.X >= WorldLocations.GetWorldMinPos().X && Parameters.Location.X <= WorldLocations.GetWorldMaxPos().X && 
+				   Parameters.Location.Y >= WorldLocations.GetWorldMinPos().Y && Parameters.Location.Y <= WorldLocations.GetWorldMaxPos().Y, TEXT("Object %s with position %s lies outside configured world boundary."), ToCStr(NetRefHandleManager->PrintObjectFromIndex((InternalReplicationIndex))), *Parameters.Location.ToString());
+
+		WorldLocations.InitObjectInfoCache(InternalReplicationIndex);
+		WorldLocations.SetObjectInfo(InternalReplicationIndex, Parameters.Location, 0.0f);
+
+		GetReplicationSystem()->SetPrioritizer(DestructionInfoHandle, DefaultSpatialNetObjectPrioritizerHandle);
+	}
+
+	UE_LOG_REPLICATIONBRIDGE(Verbose, TEXT("UReplicationBridge::StoreDestructionInfo %s (InternalIndex: %u ) for %s GroupIndex: %u"), *DestructionInfoHandle.ToString(), InternalReplicationIndex,  *PrintObjectFromNetRefHandle(Handle), LevelGroupHandle.GetGroupIndex());
+	
+	return DestructionInfoHandle;
+}
+
+bool UReplicationBridge::IsReplicatedHandle(FNetRefHandle Handle) const
+{
+	return Handle.IsValid() && NetRefHandleManager->IsNetRefHandleAssigned(Handle);
+}
+
+void UReplicationBridge::SetNetPushIdOnFragments(const TArrayView<const UE::Net::FReplicationFragment*const>& Fragments, const UE::Net::Private::FNetPushObjectHandle& PushHandle)
+{
+#if WITH_PUSH_MODEL
+	using namespace UE::Net;
+
+	constexpr uint32 MaxFragmentOwnerCount = 1U;
+	UObject* FragmentOwners[MaxFragmentOwnerCount] = {};
+	FReplicationStateOwnerCollector FragmentOwnerCollector(FragmentOwners, MaxFragmentOwnerCount);
+	UObject* Instance = nullptr;
+	for (const FReplicationFragment* Fragment : Fragments)
+	{
+		if (!EnumHasAnyFlags(Fragment->GetTraits(), EReplicationFragmentTraits::HasPushBasedDirtiness))
+		{
+			continue;
+		}
+
+		FragmentOwnerCollector.Reset();
+		Fragment->CollectOwner(&FragmentOwnerCollector);
+		UObject* FragmentOwner = (FragmentOwnerCollector.GetOwnerCount() > 0 ? FragmentOwnerCollector.GetOwners()[0] : static_cast<UObject*>(nullptr));
+		if (FragmentOwner != nullptr && FragmentOwner != Instance)
+		{
+			Instance = FragmentOwner;
+			UE_NET_IRIS_SET_PUSH_ID(Instance, PushHandle);
+		}
+	}
+#endif
+}
+
+void UReplicationBridge::ClearNetPushIdOnFragments(const TArrayView<const UE::Net::FReplicationFragment*const>& Fragments)
+{
+#if WITH_PUSH_MODEL
+	using namespace UE::Net;
+
+	constexpr uint32 MaxFragmentOwnerCount = 1U;
+	UObject* FragmentOwners[MaxFragmentOwnerCount] = {};
+	FReplicationStateOwnerCollector FragmentOwnerCollector(FragmentOwners, MaxFragmentOwnerCount);
+	UObject* Instance = nullptr;
+	for (const FReplicationFragment* Fragment : Fragments)
+	{
+		if (!EnumHasAnyFlags(Fragment->GetTraits(), EReplicationFragmentTraits::HasPushBasedDirtiness))
+		{
+			continue;
+		}
+
+		FragmentOwnerCollector.Reset();
+		Fragment->CollectOwner(&FragmentOwnerCollector);
+		UObject* FragmentOwner = (FragmentOwnerCollector.GetOwnerCount() > 0 ? FragmentOwnerCollector.GetOwners()[0] : static_cast<UObject*>(nullptr));
+		if (FragmentOwner != nullptr && FragmentOwner != Instance)
+		{
+			Instance = FragmentOwner;
+			if (IsValid(FragmentOwner))
+			{
+				UE_NET_IRIS_CLEAR_PUSH_ID(FragmentOwner);
+			}
+		}
+	}
+#endif
+}
+
+void UReplicationBridge::NotifyStreamingLevelUnload(const UObject* Level)
+{
+	// Destroy group associated with level
+	UE::Net::FNetObjectGroupHandle LevelGroupHandle;
+	if (LevelGroups.RemoveAndCopyValue(FObjectKey(Level), LevelGroupHandle))
+	{
+		RemoveDestructionInfosForGroup(LevelGroupHandle);
+		ReplicationSystem->DestroyGroup(LevelGroupHandle);
+	}
+}
+
+UE::Net::FNetObjectGroupHandle UReplicationBridge::CreateLevelGroup(const UObject* Level, FName PackageName)
+{
+	using namespace UE::Net;
+
+	FNetObjectGroupHandle LevelGroupHandle = ReplicationSystem->CreateGroup(PackageName);
+	if (ensure(LevelGroupHandle.IsValid()))
+	{
+		ReplicationSystem->AddExclusionFilterGroup(LevelGroupHandle);
+		LevelGroups.Emplace(FObjectKey(Level), LevelGroupHandle);
+	}
+
+	return LevelGroupHandle;
+}
+
+UE::Net::FNetObjectFactoryId UReplicationBridge::GetNetObjectFactoryId(FNetRefHandle RefHandle) const
+{
+	return NetRefHandleManager->GetReplicatedObjectData(NetRefHandleManager->GetInternalIndex(RefHandle)).NetFactoryId;
+}
+
+UE::Net::FNetObjectGroupHandle UReplicationBridge::GetLevelGroup(const UObject* Level) const
+{
+	const UE::Net::FNetObjectGroupHandle* LevelGroupHandle = LevelGroups.Find(FObjectKey(Level));
+	return (LevelGroupHandle != nullptr ? *LevelGroupHandle : UE::Net::FNetObjectGroupHandle());
+}
+
+const TMap<FObjectKey, UE::Net::FNetObjectGroupHandle>& UReplicationBridge::GetAllLevelGroups() const
+{
+	return LevelGroups;
+}
+
+void UReplicationBridge::DestroyGlobalNetHandle(UE::Net::Private::FInternalNetRefIndex InternalReplicationIndex)
+{
+	using namespace UE::Net;
+	using namespace UE::Net::Private;
+
+	const FNetRefHandleManager::FReplicatedObjectData& ObjectData = NetRefHandleManager->GetReplicatedObjectDataNoCheck(InternalReplicationIndex);
+	if (ObjectData.NetHandle.IsValid())
+	{
+		FNetHandleDestroyer::DestroyNetHandle(ObjectData.NetHandle);
+	}
+}
+
+void UReplicationBridge::ClearNetPushIds(UE::Net::Private::FInternalNetRefIndex InternalReplicationIndex)
+{
+#if WITH_PUSH_MODEL
+	using namespace UE::Net;
+	using namespace UE::Net::Private;
+
+	const FNetRefHandleManager::FReplicatedObjectData& ObjectData = NetRefHandleManager->GetReplicatedObjectDataNoCheck(InternalReplicationIndex);
+	if (const FReplicationInstanceProtocol* InstanceProtocol = ObjectData.InstanceProtocol)
+	{
+		if (EnumHasAnyFlags(InstanceProtocol->InstanceTraits, EReplicationInstanceProtocolTraits::HasPushBasedDirtiness))
+		{
+			TArrayView<const FReplicationFragment* const> Fragments(InstanceProtocol->Fragments, InstanceProtocol->FragmentCount);
+			ClearNetPushIdOnFragments(Fragments);
+		}
+	}
+#endif
+}
+
+FString UReplicationBridge::PrintObjectFromNetRefHandle(FNetRefHandle RefHandle) const
+{
+	return NetRefHandleManager->PrintObjectFromNetRefHandle(RefHandle);
+}
+
+void UReplicationBridge::PreSeamlessTravelGarbageCollect()
+{
+	RemoveDestructionInfosForGroup(UE::Net::FNetObjectGroupHandle());
+
+	for (const TPair<FObjectKey, UE::Net::FNetObjectGroupHandle>& LevelAndGroup : LevelGroups)
+	{
+		ReplicationSystem->DestroyGroup(LevelAndGroup.Value);
+	}
+	LevelGroups.Empty();
+
+	OnPreSeamlessTravelGarbageCollect();
+}
+
+void UReplicationBridge::OnPreSeamlessTravelGarbageCollect()
+{
+}
+
+void UReplicationBridge::PostSeamlessTravelGarbageCollect()
+{
+	OnPostSeamlessTravelGarbageCollect();
+}
+
+void UReplicationBridge::OnPostSeamlessTravelGarbageCollect()
+{
+}
+
+namespace UE::Net::Private
+{
+
+#if UE_NET_ASYNCLOADING_DEBUG
+void FAsyncLoadingSimulator::AddStalledObject(UReplicationBridge* Bridge, FNetRefHandle ReplicatedObject)
+{
+	check(Bridge);
+	if (ReplicatedObject.IsValid())
+	{
+		Bridge->DebugObjectsForcedToStall.AddUnique(ReplicatedObject);
+	}
+}
+
+void FAsyncLoadingSimulator::RemoveStalledObject(UReplicationBridge* Bridge, FNetRefHandle ReplicatedObject)
+{
+	check(Bridge);
+	Bridge->DebugObjectsForcedToStall.Remove(ReplicatedObject);
+}
+
+void FAsyncLoadingSimulator::RemoveAllStalledObjects(UReplicationBridge* Bridge)
+{
+	check(Bridge);
+	Bridge->DebugObjectsForcedToStall.Empty();
+}
+
+#endif // UE_NET_ASYNCLOADING_DEBUG
+
+} // namespace UE::Net::Private
\ No newline at end of file
diff --git a/Engine/Source/Runtime/Net/Iris/Public/Iris/ReplicationSystem/Filtering/FilterOutNetObjectFilter.h b/Engine/Source/Runtime/Net/Iris/Public/Iris/ReplicationSystem/Filtering/FilterOutNetObjectFilter.h
new file mode 100644
--- /dev/null
+++ b/Engine/Source/Runtime/Net/Iris/Public/Iris/ReplicationSystem/Filtering/FilterOutNetObjectFilter.h
@@ -0,0 +1,27 @@
+// Copyright Epic Games, Inc. All Rights Reserved.
+
+#pragma once
+
+#include "Iris/ReplicationSystem/Filtering/NetObjectFilter.h"
+#include "FilterOutNetObjectFilter.generated.h"
+
+UCLASS(transient, MinimalAPI)
+class UFilterOutNetObjectFilterConfig final : public UNetObjectFilterConfig
+{
+	GENERATED_BODY()
+};
+
+UCLASS()
+class UFilterOutNetObjectFilter final : public UNetObjectFilter
+{
+	GENERATED_BODY()
+
+protected:
+	// UNetObjectFilter interface
+	IRISCORE_API virtual void OnInit(const FNetObjectFilterInitParams&) override;
+	virtual void OnDeinit() override {}
+	virtual void OnMaxInternalNetRefIndexIncreased(uint32 NewMaxInternalIndex) override {}
+	IRISCORE_API virtual bool AddObject(uint32 ObjectIndex, FNetObjectFilterAddObjectParams&) override;
+	IRISCORE_API virtual void RemoveObject(uint32 ObjectIndex, const FNetObjectFilteringInfo&) override;
+	IRISCORE_API virtual void Filter(FNetObjectFilteringParams&) override;
+};
diff --git a/Engine/Source/Runtime/Net/Iris/Public/Iris/ReplicationSystem/Filtering/NopNetObjectFilter.h b/Engine/Source/Runtime/Net/Iris/Public/Iris/ReplicationSystem/Filtering/NopNetObjectFilter.h
new file mode 100644
--- /dev/null
+++ b/Engine/Source/Runtime/Net/Iris/Public/Iris/ReplicationSystem/Filtering/NopNetObjectFilter.h
@@ -0,0 +1,27 @@
+// Copyright Epic Games, Inc. All Rights Reserved.
+
+#pragma once
+
+#include "Iris/ReplicationSystem/Filtering/NetObjectFilter.h"
+#include "NopNetObjectFilter.generated.h"
+
+UCLASS(transient, MinimalAPI)
+class UNopNetObjectFilterConfig final : public UNetObjectFilterConfig
+{
+	GENERATED_BODY()
+};
+
+UCLASS()
+class UNopNetObjectFilter final : public UNetObjectFilter
+{
+	GENERATED_BODY()
+
+protected:
+	// UNetObjectFilter interface
+	IRISCORE_API virtual void OnInit(const FNetObjectFilterInitParams&) override;
+	virtual void OnDeinit() override {}
+	virtual void OnMaxInternalNetRefIndexIncreased(uint32 NewMaxInternalIndex) override {}
+	IRISCORE_API virtual bool AddObject(uint32 ObjectIndex, FNetObjectFilterAddObjectParams&) override;
+	IRISCORE_API virtual void RemoveObject(uint32 ObjectIndex, const FNetObjectFilteringInfo&) override;
+	IRISCORE_API virtual void Filter(FNetObjectFilteringParams&) override;
+};
diff --git a/Engine/Source/Runtime/Net/Iris/Public/Iris/ReplicationSystem/ReplicationBridge.h b/Engine/Source/Runtime/Net/Iris/Public/Iris/ReplicationSystem/ReplicationBridge.h
new file mode 100644
--- /dev/null
+++ b/Engine/Source/Runtime/Net/Iris/Public/Iris/ReplicationSystem/ReplicationBridge.h
@@ -0,0 +1,427 @@
+// Copyright Epic Games, Inc. All Rights Reserved.
+
+#pragma once
+
+#include "Containers/ContainersFwd.h"
+#include "Containers/Map.h"
+
+#include "HAL/Platform.h"
+
+#include "Iris/Core/NetObjectReference.h"
+
+#include "Iris/IrisConfig.h"
+
+#include "Iris/ReplicationSystem/NetObjectFactory.h"
+#include "Iris/ReplicationSystem/NetObjectFactoryRegistry.h"
+#include "Iris/ReplicationSystem/NetObjectGroupHandle.h"
+#include "Iris/ReplicationSystem/NetRefHandle.h"
+#include "Iris/ReplicationSystem/ObjectReferenceTypes.h"
+#include "Iris/ReplicationSystem/ReplicationBridgeTypes.h"
+
+#include "Misc/EnumClassFlags.h"
+
+#include "Net/Core/NetHandle/NetHandle.h"
+
+#include "UObject/ObjectKey.h"
+#include "UObject/ObjectMacros.h"
+
+#include "ReplicationBridge.generated.h"
+
+class UObjectReplicationBridge;
+class UReplicationSystem;
+class UNetDriver;
+
+namespace UE::Net
+{
+	enum class ENetRefHandleError : uint32;
+
+	struct FNetDependencyInfo;
+	struct FReplicationInstanceProtocol;
+	struct FReplicationProtocol;
+
+	class FNetBitStreamReader;
+	class FNetBitStreamWriter;
+	class FNetSerializationContext;
+	class FNetTokenStoreState;	
+	class FReplicationFragment;
+
+	namespace Private
+	{
+		typedef uint32 FInternalNetRefIndex;
+		class FNetRefHandleManager;
+		class FNetObjectGroups;
+		class FNetPushObjectHandle;
+		class FObjectReferenceCache;
+		class FReplicationProtocolManager;
+		class FReplicationReader;
+		class FReplicationStateDescriptorRegistry;
+		class FReplicationSystemImpl;
+		class FReplicationSystemInternal;
+		class FReplicationWriter;
+		
+		struct FChangeMaskCache;
+		struct FCreateNetObjectParams;
+		struct FAsyncLoadingSimulator;
+	}
+
+	typedef TArray<FNetDependencyInfo, TInlineAllocator<32> > FNetDependencyInfoArray;
+}
+
+#define UE_LOG_BRIDGEID(Category, Verbosity, Format, ...)  UE_LOG(Category, Verbosity, TEXT("RepBridge(%u)::") Format, GetReplicationSystem()->GetId(), ##__VA_ARGS__)
+
+struct FReplicationBridgeSerializationContext
+{
+	FReplicationBridgeSerializationContext(UE::Net::FNetSerializationContext& InSerialiazationContext, uint32 InConnectionId, bool bInIsDestructionInfo = false);
+
+	UE::Net::FNetSerializationContext& SerializationContext;
+	uint32 ConnectionId;
+	bool bIsDestructionInfo;
+};
+
+
+
+UCLASS(Transient, MinimalAPI)
+class UReplicationBridge : public UObject
+{
+	GENERATED_BODY()
+
+protected:
+	using FNetHandle = UE::Net::FNetHandle;
+	using FNetRefHandle = UE::Net::FNetRefHandle;
+	using FNetDependencyInfoArray = UE::Net::FNetDependencyInfoArray;
+
+public:
+	IRISCORE_API UReplicationBridge();
+	IRISCORE_API virtual ~UReplicationBridge();
+
+
+	/** The destruction info needed to replicate the destruction event later. */
+	struct FDestructionParameters
+	{
+		/** The location of the object. Used for distance based prioritization. */
+		FVector Location;
+		/** The level the object is placed in. */
+		const UObject* Level = nullptr;
+		/** Whether to use distance based priority for the destruction of the object. */
+		bool bUseDistanceBasedPrioritization = false;
+		/** The NetFactory that the replicated object would be assigned to. */
+		UE::Net::FNetObjectFactoryId NetFactoryId = UE::Net::InvalidNetObjectFactoryId;
+	};
+
+	enum class ESubObjectInsertionOrder : uint8
+	{
+		None,
+		/** Insert the subobject so it will replicate before the other subobject. */
+		ReplicateWith,
+		/** Insert the subobject at the start of the list so it can be created and replicated first */
+		InsertAtStart,
+	};
+
+	/**
+	* Stop replicating the NetObject associated with the handle and mark the handle to be destroyed.
+	* If EEndReplication::TearOff is set the remote instance will be Torn-off rather than being destroyed on the receiving end, after the call, any state changes will not be replicated
+	* If EEndReplication::Flush is set all pending states will be delivered before the remote instance is destroyed, final state will be immediately copied so it is safe to remove the object after this call
+	* If EEndReplication::Destroy is set the remote instance will be destroyed, if this is set for a static instance and the EndReplicationParameters are set a permanent destruction info will be added
+	* Dynamic instances are always destroyed unless the TearOff flag is set.
+	*/
+	IRISCORE_API void StopReplicatingNetRefHandle(FNetRefHandle Handle, EEndReplicationFlags EndReplicationFlags);
+
+	/** Store destruction info for the referenced object. */
+	IRISCORE_API FNetRefHandle StoreDestructionInfo(FNetRefHandle Handle, const FDestructionParameters& Parameters);
+
+	/** Returns true if the handle is replicated. */
+	IRISCORE_API bool IsReplicatedHandle(FNetRefHandle Handle) const;
+
+	/** Get the group associated with the level in order to control connection filtering for it. */
+	IRISCORE_API UE::Net::FNetObjectGroupHandle GetLevelGroup(const UObject* Level) const;
+
+	/** Get all level groups. The object key is to a ULevel. The TMap must be processed immediately or copied- DO NOT CACHE. */
+	IRISCORE_API const TMap<FObjectKey, UE::Net::FNetObjectGroupHandle>& GetAllLevelGroups() const;
+
+	/** Returns true when we are in the middle of processing incoming data. */
+	bool IsInReceiveUpdate() const { return bInReceiveUpdate; }
+
+	/** Print common information about this handle and the object it is mapped to */
+	[[nodiscard]] IRISCORE_API FString PrintObjectFromNetRefHandle(FNetRefHandle RefHandle) const;
+
+protected:
+	/** Initializes the bridge. Is called during ReplicationSystem initialization. */
+	IRISCORE_API virtual void Initialize(UReplicationSystem* InReplicationSystem);
+
+	/** Deinitializes the bridge. Is called during ReplicationSystem deinitialization. */
+	IRISCORE_API virtual void Deinitialize();
+
+	/** Invoked before ReplicationSystem copies dirty state data. */
+	virtual void PreSendUpdate() {}
+
+	/** Invoked when the ReplicationSystem starts the PreSendUpdate tick. */
+	virtual void OnStartPreSendUpdate() {}
+	
+	/** Invoked after we sent data to all connections. */
+	virtual void OnPostSendUpdate() {}
+
+	/** Invoked after we processed all incoming data */
+	virtual void OnPostReceiveUpdate() {}
+	
+	/** Invoked before ReplicationSystem copies dirty state data for a single replicated object. */
+	virtual void PreSendUpdateSingleHandle(FNetRefHandle Handle) {}
+
+	/** Update world locations in FWorldLocations for objects that support it. */
+	virtual void UpdateInstancesWorldLocation() {}
+
+	// Remote interface, invoked from Replication code during serialization
+	
+	/**
+	 * Write data required to instantiate NetObject remotely to bitstream.
+	 * @param Context The serialization context parameters.
+	 * @param Handle The handle of the object to write creation data for.
+	 */
+	IRISCORE_API virtual bool WriteNetRefHandleCreationInfo(FReplicationBridgeSerializationContext& Context, FNetRefHandle Handle);
+
+	/**
+	 * Cache info required to allow deferred writing of NetRefHandleCreationInfo
+	 * @param Handle The handle of the object to store creation data for.
+	 * return whether cached data is stored or not.
+	*/
+	IRISCORE_API virtual bool CacheNetRefHandleCreationInfo(FNetRefHandle Handle);
+
+	/** Read data required to instantiate NetObject from bitstream. */
+	IRISCORE_API virtual FReplicationBridgeCreateNetRefHandleResult CreateNetRefHandleFromRemote(FNetRefHandle RootObjectOfSubObject, FNetRefHandle WantedNetHandle, FReplicationBridgeSerializationContext& Context);
+
+	/** Invoked right before we apply the state for a new received subobject but after we have applied state on the root object. */
+ 	IRISCORE_API virtual void SubObjectCreatedFromReplication(UE::Net::Private::FInternalNetRefIndex RootObjectIndex,FNetRefHandle SubObjectCreated);
+
+	/** Invoked after we have applied the initial state for an object.*/
+	virtual void PostApplyInitialState(UE::Net::Private::FInternalNetRefIndex InternalObjectIndex) {}
+
+	/**
+	 * Called when the instance is detached from the protocol on request by the remote. 
+	 * @param Handle The handle of the object to destroy or tear off.
+	 * @param DestroyReason Reason for destroying the instance. 
+	 * @param DestroyFlags Special flags such as whether the instance may be destroyed when reason is TearOff, which may revert to destroying, or Destroy.
+	 */
+	IRISCORE_API virtual void DetachInstanceFromRemote(FNetRefHandle Handle, EReplicationBridgeDestroyInstanceReason DestroyReason, EReplicationBridgeDestroyInstanceFlags DestroyFlags,  UE::Net::FNetObjectFactoryId NetFactoryId);
+
+	/** Called when we detach instance protocol from the local instance */
+	IRISCORE_API virtual void DetachInstance(FNetRefHandle Handle);
+
+	/** Invoked post garbage collect to allow us to detect stale objects */
+	IRISCORE_API virtual void PruneStaleObjects();
+
+	/** Invoked when we start to replicate an object for a specific connection to fill in any initial dependencies */
+	IRISCORE_API virtual void GetInitialDependencies(FNetRefHandle Handle, FNetDependencyInfoArray& OutDependencies) const;
+
+	/** Returns if the bridge is allowed to create new destruction info at this moment. */
+	virtual bool CanCreateDestructionInfo() const { return true; }
+
+	/** Called just prior to garbage collecting/destroying the previous world tied to a NetDriver during seamless travel. Subclasses can implement virtual method OnPreSeamlessTravelGarbageCollect. */
+	void PreSeamlessTravelGarbageCollect();
+
+	/** Called just after garbage collecting/destroying the previous world tied to a NetDriver during seamless travel. Subclasses can implement virtual method OnPostSeamlessTravelGarbageCollect. */
+	void PostSeamlessTravelGarbageCollect();
+
+protected:
+
+	// Forward calls to internal operations that we allow replication bridges to access
+
+	/** Create a local NetRefHandle / NetObject using the ReplicationProtocol. */
+	IRISCORE_API FNetRefHandle InternalCreateNetObject(FNetRefHandle AllocatedHandle, FNetHandle GlobalHandle, const UE::Net::Private::FCreateNetObjectParams& Params);
+
+	/** Create a NetRefHandle / NetObject on request from the authoritative end. */
+	IRISCORE_API FNetRefHandle InternalCreateNetObjectFromRemote(FNetRefHandle WantedNetHandle, const UE::Net::Private::FCreateNetObjectParams& Params);
+
+	/** Attach instance to NetRefHandle. */
+	IRISCORE_API void InternalAttachInstanceToNetRefHandle(FNetRefHandle RefHandle, bool bBindInstanceProtocol, UE::Net::FReplicationInstanceProtocol* InstanceProtocol, UObject* Instance, FNetHandle NetHandle);
+
+	/** Detach instance from NetRefHandle and destroy the instance protocol. */
+	IRISCORE_API void InternalDetachInstanceFromNetRefHandle(FNetRefHandle Handle);
+
+	/** Destroy the handle and all internal book keeping associated with it. */
+	IRISCORE_API void InternalDestroyNetObject(FNetRefHandle Handle);
+	
+	/** Add SubObjectHandle as SubObject to OwnerHandle. */
+	IRISCORE_API void InternalAddSubObject(FNetRefHandle OwnerHandle, FNetRefHandle SubObjectHandle, FNetRefHandle InsertRelativeToSubObjectHandle, ESubObjectInsertionOrder InsertionOrder);
+
+	inline UE::Net::Private::FReplicationProtocolManager* GetReplicationProtocolManager() const { return ReplicationProtocolManager; }
+	inline UReplicationSystem* GetReplicationSystem() const { return ReplicationSystem; }
+	inline UE::Net::Private::FReplicationStateDescriptorRegistry* GetReplicationStateDescriptorRegistry() const { return ReplicationStateDescriptorRegistry; }
+	inline UE::Net::Private::FObjectReferenceCache* GetObjectReferenceCache() const { return ObjectReferenceCache; }
+
+	/** Return the NetFactoryId assigned to a replicated object. */
+	IRISCORE_API UE::Net::FNetObjectFactoryId GetNetObjectFactoryId(FNetRefHandle RefHandle) const;
+
+	/** Creates a group for a level for object filtering purposes. */
+	IRISCORE_API UE::Net::FNetObjectGroupHandle CreateLevelGroup(const UObject* Level, FName PackageName);
+
+	/** Destroys the group associated with the level. */
+	IRISCORE_API void DestroyLevelGroup(const UObject* Level);
+
+	/** Called when destruction info is received to determine whether the instance may be destroyed. */
+	IRISCORE_API virtual bool IsAllowedToDestroyInstance(const UObject* Instance) const;
+
+	/** Called when a remote connection detected a protocol mismatch when trying to instantiate the NetRefHandle replicated object. */
+	virtual void OnProtocolMismatchReported(FNetRefHandle RefHandle, uint32 ConnectionId) {}
+
+	/** Called when a remote connection has a critical error caused by a specific NetRefHandle */
+	virtual void OnErrorWithNetRefHandleReported(UE::Net::ENetRefHandleError ErrorType, FNetRefHandle RefHandle, uint32 ConnectionId, const TArray<FNetRefHandle>& ExtraNetRefHandle) {}
+
+	/** Tell the remote connection that we detected a reading error with a specific replicated object */
+	virtual void SendErrorWithNetRefHandle(UE::Net::ENetRefHandleError ErrorType, FNetRefHandle RefHandle, uint32 ConnectionId, TConstArrayView<const FNetRefHandle> ExtraNetRefHandle = {}) {}
+
+	/** Called from PreSeamlessTravelGarbageCollect. */
+	IRISCORE_API virtual void OnPreSeamlessTravelGarbageCollect();
+
+	/** Called from PostSeamlessTravelGarbageCollect. */
+	IRISCORE_API virtual void OnPostSeamlessTravelGarbageCollect();
+
+private:
+
+	// Internal operations invoked by ReplicationSystem/ReplicationWriter
+	void ReadAndExecuteDestructionInfoFromRemote(FReplicationBridgeSerializationContext& Context);
+	void DetachSubObjectInstancesFromRemote(FNetRefHandle Handle, EReplicationBridgeDestroyInstanceReason DestroyReason, EReplicationBridgeDestroyInstanceFlags DestroyFlags);
+	void DestroyNetObjectFromRemote(FNetRefHandle Handle, EReplicationBridgeDestroyInstanceReason DestroyReason, EReplicationBridgeDestroyInstanceFlags DestroyFlags);
+
+	enum class EPendingEndReplicationImmediate : uint8
+	{
+		Yes,
+		No,
+	};
+
+	// Adds the Handle to the list of handles pending deferred EndReplication, if bIsImmediate is true the object will be destroyed after the next update, otherwise
+	// it will be kept around until the handle is no longer ref-counted by any connection. It will however be removed from the set of scopeable objects after the first update so new connections will not add it to their scope.
+	void AddPendingEndReplication(FNetRefHandle Handle, EEndReplicationFlags DestroyFlags, EPendingEndReplicationImmediate Immediate = EPendingEndReplicationImmediate::No);
+
+	FReplicationBridgeCreateNetRefHandleResult CallCreateNetRefHandleFromRemote(FNetRefHandle RootObjectOfSubObject, FNetRefHandle WantedNetHandle, FReplicationBridgeSerializationContext& Context);
+	void CallPreSendUpdate(float DeltaSeconds);	
+	void CallPreSendUpdateSingleHandle(FNetRefHandle Handle);
+	void CallUpdateInstancesWorldLocation();
+	bool CallCacheNetRefHandleCreationInfo(FNetRefHandle Handle);
+	bool CallWriteNetRefHandleCreationInfo(FReplicationBridgeSerializationContext& Context, FNetRefHandle Handle);
+	bool CallWriteNetRefHandleDestructionInfo(FReplicationBridgeSerializationContext& Context, FNetRefHandle Handle);
+	void CallSubObjectCreatedFromReplication(UE::Net::Private::FInternalNetRefIndex RootObjectIndex, FNetRefHandle SubObjectCreated);
+	void CallPostApplyInitialState(UE::Net::Private::FInternalNetRefIndex InternalObjectIndex);
+	void CallPruneStaleObjects();
+	void CallGetInitialDependencies(FNetRefHandle Handle, FNetDependencyInfoArray& OutDependencies) const;
+	void CallDetachInstance(FNetRefHandle Handle);
+	void CallDetachInstanceFromRemote(FNetRefHandle Handle, EReplicationBridgeDestroyInstanceReason DestroyReason, EReplicationBridgeDestroyInstanceFlags DestroyFlags, UE::Net::FNetObjectFactoryId NetFactoryId);
+
+	void PreReceiveUpdate();
+	void PostReceiveUpdate();
+
+private:
+
+	void InternalFlushStateData(UE::Net::FNetSerializationContext& SerializationContext, UE::Net::Private::FChangeMaskCache& ChangeMaskCache, UE::Net::FNetBitStreamWriter& ChangeMaskWriter, uint32 InternalObjectIndex);
+	// Internal method to copy state data for Handle
+	void InternalFlushStateData(FNetRefHandle Handle);
+
+	// Internal method to copy state data for Handle and any SubObjects and mark them as being torn-off
+	void InternalTearOff(FNetRefHandle OwnerHandle);
+
+	// Destroy all SubObjects owned by provided handle
+	void InternalDestroySubObjects(FNetRefHandle OwnerHandle, EEndReplicationFlags Flags);
+
+	/**
+	 * Called from ReplicationSystem when a streaming level is about to unload.
+	 * Will remove the group associated with the level and remove destruction infos.
+	 */
+	void NotifyStreamingLevelUnload(const UObject* Level);
+
+	/**
+	 * Remove destruction infos associated with group
+	 * Passing in an invalid group handle indicates that we should remove all destruction infos
+	 */
+	void RemoveDestructionInfosForGroup(UE::Net::FNetObjectGroupHandle GroupHandle);
+
+	void DestroyLocalNetHandle(FNetRefHandle Handle, EEndReplicationFlags Flags);
+
+	// Tear-off all handles in the PendingTearOff list that has not yet been torn-off
+	void TearOffHandlesPendingTearOff();
+
+	// Update all the handles pending EndReplication
+	void UpdateHandlesPendingEndReplication();
+
+	void SetNetPushIdOnFragments(const TArrayView<const UE::Net::FReplicationFragment*const>& Fragments, const UE::Net::Private::FNetPushObjectHandle& PushHandle);
+	void ClearNetPushIdOnFragments(const TArrayView<const UE::Net::FReplicationFragment*const>& Fragments);
+
+	void DestroyGlobalNetHandle(UE::Net::Private::FInternalNetRefIndex InternalReplicationIndex);
+	void ClearNetPushIds(UE::Net::Private::FInternalNetRefIndex InternalReplicationIndex);
+
+	friend UReplicationSystem;
+	friend UE::Net::Private::FReplicationSystemImpl;
+	friend UE::Net::Private::FReplicationSystemInternal;
+	friend UE::Net::Private::FReplicationWriter;
+	friend UE::Net::Private::FReplicationReader;
+	friend UObjectReplicationBridge;
+
+	UReplicationSystem* ReplicationSystem;
+	UE::Net::Private::FReplicationProtocolManager* ReplicationProtocolManager;
+	UE::Net::Private::FReplicationStateDescriptorRegistry* ReplicationStateDescriptorRegistry;
+	UE::Net::Private::FNetRefHandleManager* NetRefHandleManager;
+	UE::Net::Private::FObjectReferenceCache* ObjectReferenceCache;
+	UE::Net::Private::FNetObjectGroups* Groups;
+
+	TMap<FObjectKey, UE::Net::FNetObjectGroupHandle> LevelGroups;
+
+	/** Tracks if we are in the middle of processing incoming data */
+	bool bInReceiveUpdate = false;
+
+	/** List of replicated objects that requested to stop replicating while we were in ReceiveUpdate */
+	TMap<FNetRefHandle, EEndReplicationFlags> HandlesToStopReplicating;
+
+private:
+
+	// DestructionInfos
+	struct FDestructionInfo
+	{
+		UE::Net::FNetObjectReference StaticRef;
+		UE::Net::FNetObjectGroupHandle LevelGroupHandle;
+		UE::Net::FNetObjectFactoryId NetFactoryId = UE::Net::InvalidNetObjectFactoryId;
+		UE::Net::Private::FInternalNetRefIndex InternalReplicationIndex;
+	};
+
+	const UE::Net::FReplicationProtocol* DestructionInfoProtocol;
+	
+	// Need to track the objects with destruction infos so that we can clean them up properly
+	// We use this to be able ask remote to destroy static objects
+	TMap<FNetRefHandle, FDestructionInfo> StaticObjectsPendingDestroy;
+
+	struct FPendingEndReplicationInfo
+	{
+		FPendingEndReplicationInfo(FNetRefHandle InHandle, EEndReplicationFlags InDestroyFlags, EPendingEndReplicationImmediate InImmediate) : Handle(InHandle), DestroyFlags(InDestroyFlags), Immediate(InImmediate) {}
+
+		FNetRefHandle Handle;
+		EEndReplicationFlags DestroyFlags;
+		EPendingEndReplicationImmediate Immediate;
+	};
+	TArray<FPendingEndReplicationInfo> HandlesPendingEndReplication;
+
+#if UE_NET_ASYNCLOADING_DEBUG
+	friend UE::Net::Private::FAsyncLoadingSimulator;
+	/** 
+	 * Debug array that forces objects to be stalled, as if async loading external references.
+	 * Only enabled when net.Iris.AsyncLoading.ForceObjectsToStall is enabled.
+	 */
+	TArray<FNetRefHandle> DebugObjectsForcedToStall;
+#endif
+
+};
+
+namespace UE::Net::Private
+{
+#if UE_NET_ASYNCLOADING_DEBUG
+	struct FAsyncLoadingSimulator
+	{
+		IRISCORE_API static void AddStalledObject(UReplicationBridge* Bridge, FNetRefHandle ReplicatedObject);
+		IRISCORE_API static void RemoveStalledObject(UReplicationBridge* Bridge, FNetRefHandle ReplicatedObject);
+		IRISCORE_API static void RemoveAllStalledObjects(UReplicationBridge* Bridge);
+	};
+#endif
+
+}
+
+inline FReplicationBridgeSerializationContext::FReplicationBridgeSerializationContext(UE::Net::FNetSerializationContext& InSerialiazationContext, uint32 InConnectionId, bool bInIsDestructionInfo)
+: SerializationContext(InSerialiazationContext)
+, ConnectionId(InConnectionId)
+, bIsDestructionInfo(bInIsDestructionInfo)
+{
+}
diff --git a/Engine/Source/Runtime/Net/Iris/Public/Iris/Serialization/PolymorphicNetSerializer.h b/Engine/Source/Runtime/Net/Iris/Public/Iris/Serialization/PolymorphicNetSerializer.h
new file mode 100644
--- /dev/null
+++ b/Engine/Source/Runtime/Net/Iris/Public/Iris/Serialization/PolymorphicNetSerializer.h
@@ -0,0 +1,87 @@
+// Copyright Epic Games, Inc. All Rights Reserved.
+
+#pragma once
+
+#include "NetSerializer.h"
+#include "UObject/ObjectMacros.h"
+#include "Iris/ReplicationState/ReplicationStateDescriptor.h"
+#include "Containers/Array.h"
+
+#include "PolymorphicNetSerializer.generated.h"
+
+class UScriptStruct;
+namespace UE::Net
+{
+	struct FReplicationStateDescriptor;
+}
+
+namespace UE::Net
+{
+
+struct FPolymorphicNetSerializerScriptStructCache
+{
+	enum : uint32 { RegisteredTypeBits = 5U };
+	enum : uint32 { MaxRegisteredTypeCount = (1U << RegisteredTypeBits) - 1U };
+	enum : uint32 { InvalidTypeIndex = 0U };
+
+	struct FTypeInfo
+	{
+		const UScriptStruct* ScriptStruct;
+		TRefCountPtr<const FReplicationStateDescriptor> Descriptor;
+	};
+
+	IRISCORE_API void InitForType(const UScriptStruct* InScriptStruct);
+	inline const uint32 GetTypeIndex(const UScriptStruct* ScriptStruct) const;
+	inline const FTypeInfo* GetTypeInfo(uint32 TypeIndex) const;
+	inline bool CanHaveNetReferences() const { return EnumHasAnyFlags(CommonTraits, EReplicationStateTraits::HasObjectReference); }
+
+private:
+	TArray<FTypeInfo> RegisteredTypes;
+	EReplicationStateTraits CommonTraits = EReplicationStateTraits::None;
+	int32 EstimatedScriptStructCount = 4096;
+};
+
+}
+
+USTRUCT()
+struct FPolymorphicStructNetSerializerConfig : public FNetSerializerConfig
+{
+	GENERATED_BODY()
+
+	UE::Net::FPolymorphicNetSerializerScriptStructCache RegisteredTypes;
+};
+
+USTRUCT()
+struct FPolymorphicArrayStructNetSerializerConfig : public FPolymorphicStructNetSerializerConfig
+{
+	GENERATED_BODY()
+};
+
+namespace UE::Net
+{
+
+const FPolymorphicNetSerializerScriptStructCache::FTypeInfo* FPolymorphicNetSerializerScriptStructCache::GetTypeInfo(uint32 TypeIndex) const
+{
+	static_assert(InvalidTypeIndex == 0U, "Expected InvalidTypeIndex to be 0");
+	if ((TypeIndex - 1U) >= static_cast<uint32>(RegisteredTypes.Num()))
+	{
+		return nullptr;
+	}
+	return &RegisteredTypes.GetData()[TypeIndex - 1U];
+}
+
+const uint32 FPolymorphicNetSerializerScriptStructCache::GetTypeIndex(const UScriptStruct* ScriptStruct) const
+{
+	if (ScriptStruct)
+	{
+		const int32 ArrayIndex =  RegisteredTypes.IndexOfByPredicate([&ScriptStruct](const FTypeInfo& TypeInfo) { return TypeInfo.ScriptStruct == ScriptStruct; } );
+		if (ArrayIndex != INDEX_NONE)
+		{
+			return static_cast<uint32>(ArrayIndex + 1);
+		}
+	}
+
+	return InvalidTypeIndex;
+}
+
+}
