Iris - Fixed bug that could prevent a default subobject from receiving a large RPC. Added some tests for RPCs on components.

diff --git a/Engine/Source/Runtime/Net/Iris/Private/Iris/ReplicationSystem/NetBlob/PartialNetBlob.cpp b/Engine/Source/Runtime/Net/Iris/Private/Iris/ReplicationSystem/NetBlob/PartialNetBlob.cpp
--- a/Engine/Source/Runtime/Net/Iris/Private/Iris/ReplicationSystem/NetBlob/PartialNetBlob.cpp
+++ b/Engine/Source/Runtime/Net/Iris/Private/Iris/ReplicationSystem/NetBlob/PartialNetBlob.cpp
@@ -31,14 +31,7 @@ FPartialNetBlob::FPartialNetBlob(const FNetBlobCreationInfo& CreationInfo)
 
 TArrayView<const FNetObjectReference> FPartialNetBlob::GetNetObjectReferenceExports() const
 {
-	if (IsFirstPart() && OriginalBlob.IsValid())
-	{
-		return OriginalBlob->CallGetNetObjectReferenceExports();
-	}
-	else
-	{
-		return MakeArrayView<const FNetObjectReference>(nullptr, 0);			
-	}
+	return MakeConstArrayView<>(ObjectReferenceExportsArray);
 }
 
 TArrayView<const FNetToken> FPartialNetBlob::GetNetTokenExports() const
@@ -241,6 +234,7 @@ bool FPartialNetBlob::SplitNetBlob(FNetSerializationContext& Context, const FNet
 		PayloadSplitParams.PartBitCount = MaxPartBitCount;
 		if (Private::FNetExportContext* ExportContext = Context.GetExportContext())
 		{			
+			PayloadSplitParams.ObjectReferencesPendingExport = MakeConstArrayView<>(ExportContext->GetBatchExports().ReferencesPendingExportInCurrentBatch);
 			PayloadSplitParams.NetTokensPendingExport = MakeConstArrayView<>(ExportContext->GetBatchExports().NetTokensPendingExportInCurrentBatch);
 		}
 
@@ -293,10 +287,11 @@ void FPartialNetBlob::SplitPayload(const FPartialNetBlob::FPayloadSplitParams& S
 		PartialBlob->PartCount = static_cast<uint16>(PartialBlobCount);
 		PartialBlob->SequenceFlags = (bIsFirstPart ? ESequenceFlags::IsFirstPart : ESequenceFlags::None);
 		PartialBlob->SequenceNumber = SequenceNumber++;
-		if (bIsFirstPart && (SplitParams.NetTokensPendingExport.Num() || EnumHasAnyFlags(SplitParams.OriginalCreationInfo.Flags, ENetBlobFlags::HasExports)))
+		if (bIsFirstPart && (!SplitParams.ObjectReferencesPendingExport.IsEmpty() || !SplitParams.NetTokensPendingExport.IsEmpty()))
 		{
 			PartialBlob->OriginalBlob = SplitParams.OriginalBlob;
 			PartialBlob->CreationInfo.Flags |= ENetBlobFlags::HasExports;
+			PartialBlob->ObjectReferenceExportsArray.Append(SplitParams.ObjectReferencesPendingExport);
 			PartialBlob->NetTokenExportsArray.Append(SplitParams.NetTokensPendingExport);
 		}
 
diff --git a/Engine/Source/Runtime/Net/Iris/Public/Iris/ReplicationSystem/NetBlob/NetBlob.h b/Engine/Source/Runtime/Net/Iris/Public/Iris/ReplicationSystem/NetBlob/NetBlob.h
--- a/Engine/Source/Runtime/Net/Iris/Public/Iris/ReplicationSystem/NetBlob/NetBlob.h
+++ b/Engine/Source/Runtime/Net/Iris/Public/Iris/ReplicationSystem/NetBlob/NetBlob.h
@@ -35,7 +35,7 @@ enum class ENetBlobFlags : uint32
 	/** Used for FRawDataNetBlob derived classes to avoid duplicate serialization when splitting large blob. */
 	RawDataNetBlob = Reliable << 1U,
 
-	/** Used to indicate that this blob have ObjectReferences that might have to be exported. */
+	/** Used to indicate that this blob have ObjectReferences or NetTokens that might have to be exported. */
 	HasExports = RawDataNetBlob << 1U,
 
 	/** The blob should respect delivery order with respect to other Ordered blobs, including Reliable ones. Unreliable ordered blobs will only be sent once. */
diff --git a/Engine/Source/Runtime/Net/Iris/Public/Iris/ReplicationSystem/NetBlob/PartialNetBlob.h b/Engine/Source/Runtime/Net/Iris/Public/Iris/ReplicationSystem/NetBlob/PartialNetBlob.h
--- a/Engine/Source/Runtime/Net/Iris/Public/Iris/ReplicationSystem/NetBlob/PartialNetBlob.h
+++ b/Engine/Source/Runtime/Net/Iris/Public/Iris/ReplicationSystem/NetBlob/PartialNetBlob.h
@@ -75,7 +75,8 @@ private:
 		FNetBlobCreationInfo CreationInfo;
 		FNetBlobCreationInfo OriginalCreationInfo;
 		FNetBlob* OriginalBlob;
-		TArrayView<const FNetToken> NetTokensPendingExport;
+		TConstArrayView<FNetObjectReference> ObjectReferencesPendingExport;
+		TConstArrayView<FNetToken> NetTokensPendingExport;
 		const uint32* Payload;
 		uint32 PayloadBitCount;
 		uint32 PartBitCount;
@@ -93,6 +94,7 @@ private:
 	uint16 PayloadBitCount = 0;
 	// Use uint32 for guaranteed FNetBitStreamReader/Writer compatibility
 	TArray<uint32> Payload;
+	TArray<FNetObjectReference, TInlineAllocator<2>> ObjectReferenceExportsArray;
 	TArray<FNetToken, TInlineAllocator<4>> NetTokenExportsArray;
 
 	FNetDebugName DebugName;
