Iris
* Added a flag needed when a factory finds a static object without using ResolveObjectReference
* Allow static references to be manually added to the ObjectReferenceCache

SceneGraph
* Fix stable components not correctly cached when found on the entity via the class.

#rnx
#DisabledValidations Preflight

diff --git a/Engine/Source/Runtime/Net/Iris/Private/Iris/ReplicationSystem/ObjectReferenceCache.cpp b/Engine/Source/Runtime/Net/Iris/Private/Iris/ReplicationSystem/ObjectReferenceCache.cpp
--- a/Engine/Source/Runtime/Net/Iris/Private/Iris/ReplicationSystem/ObjectReferenceCache.cpp
+++ b/Engine/Source/Runtime/Net/Iris/Private/Iris/ReplicationSystem/ObjectReferenceCache.cpp
@@ -383,25 +383,25 @@ bool FObjectReferenceCache::CreateObjectReferenceInternal(const UObject* Object,
 				// The connection ID isn't used unless reading.
 				RenamePathForPie(UE::Net::InvalidConnectionId, RelativeObjectPath, bReading);
 				OutReference = MakeNetObjectReference(*ReplicatedOuter, StringTokenStore->GetOrCreateToken(RelativeObjectPath));
+
 				return true;
 			}
 		}
 
-		if (bIsDynamic)
-		{
-			// If the object is dynamic and we did not find an outer handle there's nothing we can do
-			return false;
-		}
-		else
+		// For static objects we can always fall back on the full pathname
+		if (bIsStatic)
 		{
-			// For static objects we can always fall back on the full pathname
 			FString ObjectPath = Object->GetPathName();
 			constexpr bool bReading = false;
 			// The connection ID isn't used unless reading.
 			RenamePathForPie(UE::Net::InvalidConnectionId, ObjectPath, bReading);
 			OutReference = MakeNetObjectReference(FNetRefHandle::GetInvalid(), StringTokenStore->GetOrCreateToken(ObjectPath));	
+			
 			return true;
 		}
+
+		// If the object is dynamic and we did not find an outer handle there's nothing we can do
+		return false;
 	}
 
 	// Generate a new RefHandle
@@ -515,66 +515,58 @@ FNetRefHandle FObjectReferenceCache::GetObjectReferenceHandleFromObject(const UO
 	return FNetRefHandle::GetInvalid();
 }
 
-void FObjectReferenceCache::AddRemoteReference(FNetRefHandle RefHandle, const UObject* Object)
+void FObjectReferenceCache::BindRemoteReference(FNetRefHandle RefHandle, const UObject* Object)
 {
 	// Only clients should get here
 	check(!IsAuthority()); 
 	check(IsValid(Object));
 	check(RefHandle.IsCompleteHandle());
 
-	if (ensure(RefHandle.IsDynamic()))
-	{
-		UE_LOG_REFERENCECACHE(Verbose, TEXT("ObjectReferenceCache::AddRemoteReference: Adding %s to ObjectToNetReferenceHandle, Object: %s (0x%p)" ), *RefHandle.ToString(), *GetNameSafe(Object), Object);
+	UE_LOG_REFERENCECACHE(Verbose, TEXT("%hs: Binding Handle: %s to Object: %s (0x%p)" ), __FUNCTION__, *RefHandle.ToString(), *GetNameSafe(Object), Object);
 
-	#if UE_NET_VALIDATE_REFERENCECACHE
-		// Dynamic object might have been referenced before it gets replicated, then it will then already be registered (with a path relative to its outer)
-		if (FNetRefHandle* ExistingHandle = ObjectToNetReferenceHandle.Find(Object))
+#if UE_NET_VALIDATE_REFERENCECACHE
+	// Dynamic object might have been referenced before it gets replicated, then it will then already be registered (with a path relative to its outer)
+	if (FNetRefHandle* ExistingHandle = ObjectToNetReferenceHandle.Find(Object))
+	{
+		if (FCachedNetObjectReference* ExistingCachedObject = ReferenceHandleToCachedReference.Find(*ExistingHandle))
 		{
-			if (FCachedNetObjectReference* ExistingCachedObject = ReferenceHandleToCachedReference.Find(*ExistingHandle))
+			if (ExistingCachedObject->NetRefHandle != RefHandle)
 			{
-				if (ExistingCachedObject->NetRefHandle != RefHandle)
-				{
-					UE_LOG_REFERENCECACHE(Warning, TEXT("ObjectReferenceCache::AddRemoteReference: Detected conflicting Ref %s: Received: %s, Object: %s (0x%p)"), *ExistingCachedObject.NetRefHandle.ToString(),
-					*RefHandle.ToString(), *GetNameSafe(Object), Object);
-				}
+				UE_LOG_REFERENCECACHE(Warning, TEXT("%hs: Detected conflicting Ref %s: Received: %s, Object: %s (0x%p)"), __FUNCTION__, *ExistingCachedObject.NetRefHandle.ToString(),
+									  *RefHandle.ToString(), *GetNameSafe(Object), Object);
 			}
 		}
-	#endif
+	}
+#endif
 
-		// The object could already be known as it has been exported as a reference prior to being replicated to this connection
-		if (FCachedNetObjectReference* CacheObjectPtr = ReferenceHandleToCachedReference.Find(RefHandle))
-		{
-			UE_LOG_REFERENCECACHE(Verbose, TEXT("ObjectReferenceCache::AddRemoteReference: %s, Object: %s (0x%p).  Patching up existing cache that held: 0x%p"), *RefHandle.ToString(), *GetNameSafe(Object), Object, CacheObjectPtr->ObjectKey);
+	// The object could already be known as it has been exported as a reference prior to being replicated to this connection
+	if (FCachedNetObjectReference* CacheObjectPtr = ReferenceHandleToCachedReference.Find(RefHandle))
+	{
+		UE_LOG_REFERENCECACHE(Verbose, TEXT("%hs: Updating handle:%s to Object: %s (0x%p). Previous ptr: (0x%p)"), __FUNCTION__, *RefHandle.ToString(), *GetNameSafe(Object), Object, CacheObjectPtr->ObjectKey);
 
-			// Patch-up object pointers
-			CacheObjectPtr->Object = MakeWeakObjectPtr(const_cast<UObject*>(Object));
-			CacheObjectPtr->ObjectKey = Object;
+		// Patch-up object pointers
+		CacheObjectPtr->Object = MakeWeakObjectPtr(const_cast<UObject*>(Object));
+		CacheObjectPtr->ObjectKey = Object;
 
-		}
-		else
-		{
-			LLM_SCOPE_BYTAG(Iris);
+	}
+	else
+	{
+		LLM_SCOPE_BYTAG(Iris);
 
-			FCachedNetObjectReference CachedObject;
+		FCachedNetObjectReference CachedObject;
 	
-			CachedObject.Object = MakeWeakObjectPtr(const_cast<UObject*>(Object));
-			CachedObject.ObjectKey = Object;
-			CachedObject.NetRefHandle = RefHandle;
-			CachedObject.RelativePath = FNetToken();
-			CachedObject.OuterNetRefHandle = FNetRefHandle::GetInvalid();
-			CachedObject.bIsPackage = false;
-			CachedObject.bNoLoad = false;
-			CachedObject.bIsBroken = false;
-			CachedObject.bIsPending = false;
-			CachedObject.bIgnoreWhenMissing = false;
+		CachedObject.Object = MakeWeakObjectPtr(const_cast<UObject*>(Object));
+		CachedObject.ObjectKey = Object;
+		CachedObject.NetRefHandle = RefHandle;
+		CachedObject.RelativePath = FNetToken();
+		CachedObject.OuterNetRefHandle = FNetRefHandle::GetInvalid();
 
-			UE_LOG_REFERENCECACHE(Verbose, TEXT("ObjectReferenceCache::AddRemoteReference: Adding %s to Cache Object: %s (0x%p).  Adding new cache"), *RefHandle.ToString(), *GetNameSafe(Object), Object);
+		UE_LOG_REFERENCECACHE(Verbose, TEXT("%hs: Adding new cache entry for: %s with Object: %s (0x%p)."), __FUNCTION__, *RefHandle.ToString(), *GetNameSafe(Object), Object);
 
-			ReferenceHandleToCachedReference.Add(RefHandle, CachedObject);
-		}
-
-		ObjectToNetReferenceHandle.Add(Object, RefHandle);
+		ReferenceHandleToCachedReference.Add(RefHandle, CachedObject);
 	}
+
+	ObjectToNetReferenceHandle.Add(Object, RefHandle);
 }
 
 void FObjectReferenceCache::AddPreRegisteredReference(FNetRefHandle RefHandle, const UObject* Object)
@@ -584,7 +576,7 @@ void FObjectReferenceCache::AddPreRegisteredReference(FNetRefHandle RefHandle, c
 		// Promote to full handle
 		FNetRefHandle CompleteRefHandle = FNetRefHandleManager::MakeNetRefHandle(RefHandle.GetId(), ReplicationSystem->GetId());
 
-		AddRemoteReference(CompleteRefHandle, Object);
+		BindRemoteReference(CompleteRefHandle, Object);
 
 		if (FCachedNetObjectReference* CacheObjectPtr = ReferenceHandleToCachedReference.Find(CompleteRefHandle))
 		{
@@ -601,108 +593,113 @@ void FObjectReferenceCache::AddPreRegisteredReference(FNetRefHandle RefHandle, c
 void FObjectReferenceCache::RemoveReference(FNetRefHandle RefHandle, const UObject* Object)
 {
 	// Cleanup cached data for dynamic instance
-	if (ensure(RefHandle.IsDynamic()))
+	const bool bIsDynamic = RefHandle.IsDynamic();
+	if (!bIsDynamic)
 	{
-		UE_LOG_REFERENCECACHE(Verbose, TEXT("ObjectReferenceCache::RemoveReference: %s, Object: %s (0x%p)"), *RefHandle.ToString(), *GetNameSafe(Object), Object);
+		ensureMsgf(!bIsDynamic, TEXT("%hs not allowed on static handles: %s - %s"), __FUNCTION__, *RefHandle.ToString(), *GetNameSafe(Object));
+		return;
+	}
 
-		// We create an iterator to avoid having to find the key twice
-		TMap<FNetRefHandle, FCachedNetObjectReference>::TKeyIterator It = ReferenceHandleToCachedReference.CreateKeyIterator(RefHandle);
-		if (It)
+
+	UE_LOG_REFERENCECACHE(Verbose, TEXT("ObjectReferenceCache::RemoveReference: %s, Object: %s (0x%p)"), *RefHandle.ToString(), *GetNameSafe(Object), Object);
+
+	// We create an iterator to avoid having to find the key twice
+	TMap<FNetRefHandle, FCachedNetObjectReference>::TKeyIterator It = ReferenceHandleToCachedReference.CreateKeyIterator(RefHandle);
+	if (It)
+	{
+		FCachedNetObjectReference& CachedItemToRemove = It.Value();
+		if (Object == nullptr)
 		{
-			FCachedNetObjectReference& CachedItemToRemove = It.Value();
-			if (Object == nullptr)
-			{
-				// Object might have been GC'd but we still want to cleanup 
-				Object = CachedItemToRemove.ObjectKey;
-			}
-			else
-			{
-				// RemoveReference can be called multiple times for the same NetRefHandle. The first call will clear out the object references.
-				ensureMsgf(CachedItemToRemove.ObjectKey == nullptr || Object == CachedItemToRemove.ObjectKey, TEXT("ObjectReferenceCache::RemoveReference: %s, Object 0x%p doesn't match cached reference 0x%p."), *RefHandle.ToString(), Object, CachedItemToRemove.ObjectKey);
-			}
+			// Object might have been GC'd but we still want to cleanup 
+			Object = CachedItemToRemove.ObjectKey;
+		}
+		else
+		{
+			// RemoveReference can be called multiple times for the same NetRefHandle. The first call will clear out the object references.
+			ensureMsgf(CachedItemToRemove.ObjectKey == nullptr || Object == CachedItemToRemove.ObjectKey, TEXT("ObjectReferenceCache::RemoveReference: %s, Object 0x%p doesn't match cached reference 0x%p."), *RefHandle.ToString(), Object, CachedItemToRemove.ObjectKey);
+		}
 
 #if UE_NET_CLEANUP_REFERENCES_WITH_DYNAMIC_OUTER
-			// We want to remove cached references that referenced the removed dynamic references
-			TArray<FNetRefHandle> SubReferences;
-			HandleToDynamicOuter.MultiFind(RefHandle, SubReferences);
-			if (SubReferences.Num())
+		// We want to remove cached references that referenced the removed dynamic references
+		TArray<FNetRefHandle> SubReferences;
+		HandleToDynamicOuter.MultiFind(RefHandle, SubReferences);
+		if (SubReferences.Num())
+		{
+			for (FNetRefHandle SubRefHandle : SubReferences)
 			{
-				for (FNetRefHandle SubRefHandle : SubReferences)
+				FCachedNetObjectReference RemovedSubObjectRef;
+				if (ReferenceHandleToCachedReference.RemoveAndCopyValue(SubRefHandle, RemovedSubObjectRef))
 				{
-					FCachedNetObjectReference RemovedSubObjectRef;
-					if (ReferenceHandleToCachedReference.RemoveAndCopyValue(SubRefHandle, RemovedSubObjectRef))
-					{
-						const UObject* SubRefObject = RemovedSubObjectRef.Object.Get();
-						UE_LOG_REFERENCECACHE(Verbose, TEXT("ObjectReferenceCache::RemoveReference: %s removed from ObjectToNetReferenceHandle. Referenced dynamic %s, Object: %s Pointer: %p"), *SubRefHandle.ToString(), *RefHandle.ToString(), SubRefObject ? *SubRefObject->GetName() : TEXT("NULL"), RemovedSubObjectRef.ObjectKey);
+					const UObject* SubRefObject = RemovedSubObjectRef.Object.Get();
+					UE_LOG_REFERENCECACHE(Verbose, TEXT("ObjectReferenceCache::RemoveReference: %s removed from ObjectToNetReferenceHandle. Referenced dynamic %s, Object: %s Pointer: %p"), *SubRefHandle.ToString(), *RefHandle.ToString(), SubRefObject ? *SubRefObject->GetName() : TEXT("NULL"), RemovedSubObjectRef.ObjectKey);
 
-						ObjectToNetReferenceHandle.Remove(RemovedSubObjectRef.ObjectKey);
-					}
+					ObjectToNetReferenceHandle.Remove(RemovedSubObjectRef.ObjectKey);
 				}
-
-				HandleToDynamicOuter.Remove(RefHandle);
 			}
+
+			HandleToDynamicOuter.Remove(RefHandle);
+		}
 #endif
 
-			// Remove if RelativePath is invalid
-			if (!CachedItemToRemove.RelativePath.IsValid())
-			{
-				UE_LOG_REFERENCECACHE(Verbose, TEXT("ObjectReferenceCache::RemoveReference: %s, Object: (0x%p). Destroying cache since Path is invalid (%s)"),
-					*RefHandle.ToString(), Object, *CachedItemToRemove.RelativePath.ToString());
-				It.RemoveCurrent();
-			}
-			else
-			{
-				UE_LOG_REFERENCECACHE(Verbose, TEXT("ObjectReferenceCache::RemoveReference: %s, Object: (0x%p). Clearing cache reference to (0x%p)"),
-					*RefHandle.ToString(), Object, CachedItemToRemove.ObjectKey);
-				CachedItemToRemove.ObjectKey = nullptr;
-				CachedItemToRemove.Object.Reset();
-			}
-		}	
+		// Remove if RelativePath is invalid
+		if (!CachedItemToRemove.RelativePath.IsValid())
+		{
+			UE_LOG_REFERENCECACHE(Verbose, TEXT("ObjectReferenceCache::RemoveReference: %s, Object: (0x%p). Destroying cache since Path is invalid (%s)"),
+				*RefHandle.ToString(), Object, *CachedItemToRemove.RelativePath.ToString());
+			It.RemoveCurrent();
+		}
+		else
+		{
+			UE_LOG_REFERENCECACHE(Verbose, TEXT("ObjectReferenceCache::RemoveReference: %s, Object: (0x%p). Clearing cache reference to (0x%p)"),
+				*RefHandle.ToString(), Object, CachedItemToRemove.ObjectKey);
+			CachedItemToRemove.ObjectKey = nullptr;
+			CachedItemToRemove.Object.Reset();
+		}
+	}	
 		
-		if (Object)
+	if (Object)
+	{
+		// We create an iterator to avoid having to find the key twice
+		TMap<const UObject*, FNetRefHandle>::TKeyIterator ObjectIt = ObjectToNetReferenceHandle.CreateKeyIterator(Object);
+		if (ObjectIt)
 		{
-			// We create an iterator to avoid having to find the key twice
-			TMap<const UObject*, FNetRefHandle>::TKeyIterator ObjectIt = ObjectToNetReferenceHandle.CreateKeyIterator(Object);
-			if (ObjectIt)
+			const FNetRefHandle ExistingRefHandle = ObjectIt.Value();
+			if (ExistingRefHandle == RefHandle)
 			{
-				const FNetRefHandle ExistingRefHandle = ObjectIt.Value();
-				if (ExistingRefHandle == RefHandle)
-				{
-					UE_LOG_REFERENCECACHE(Verbose, TEXT("ObjectReferenceCache::RemoveReference: %s removed from ObjectToNetReferenceHandle Object: (0x%p)."), 
-						*RefHandle.ToString(), Object);
-					ObjectIt.RemoveCurrent();
-				}
-				else
-				{
-					// If we have an existing handle associated with the Object that differs from the one stored in the cache it means that the client might have destroyed the instance without notifying the replication system
-					// and reused the instance, if that is the case we just leave the association as it is
-					UE_LOG_REFERENCECACHE(Verbose, TEXT("ObjectReferenceCache::RemoveReference: ExistingReference %s != RefHandle being removed: %s, Object: %s (0x%p) "), 
-						*(ExistingRefHandle.ToString()), *RefHandle.ToString(), *GetNameSafe(Object), Object);
-				}
+				UE_LOG_REFERENCECACHE(Verbose, TEXT("ObjectReferenceCache::RemoveReference: %s removed from ObjectToNetReferenceHandle Object: (0x%p)."), 
+					*RefHandle.ToString(), Object);
+				ObjectIt.RemoveCurrent();
 			}
 			else
 			{
-				UE_LOG_REFERENCECACHE(Verbose, TEXT("ObjectReferenceCache::RemoveReference: %s, Object: %s (0x%p). Skipped removing since ObjectToNetReferenceHandle had no entry"), 
-					*RefHandle.ToString(), *GetNameSafe(Object), Object);
+				// If we have an existing handle associated with the Object that differs from the one stored in the cache it means that the client might have destroyed the instance without notifying the replication system
+				// and reused the instance, if that is the case we just leave the association as it is
+				UE_LOG_REFERENCECACHE(Verbose, TEXT("ObjectReferenceCache::RemoveReference: ExistingReference %s != RefHandle being removed: %s, Object: %s (0x%p) "), 
+					*(ExistingRefHandle.ToString()), *RefHandle.ToString(), *GetNameSafe(Object), Object);
 			}
 		}
-		else if (bIrisEnableRemoveReferenceFallbackPath)
+		else
 		{
-			bool bSuccesfullyRemoved = false;
-			for (auto RefHandleIt = ObjectToNetReferenceHandle.CreateIterator(); RefHandleIt; ++RefHandleIt)
+			UE_LOG_REFERENCECACHE(Verbose, TEXT("ObjectReferenceCache::RemoveReference: %s, Object: %s (0x%p). Skipped removing since ObjectToNetReferenceHandle had no entry"), 
+				*RefHandle.ToString(), *GetNameSafe(Object), Object);
+		}
+	}
+	else if (bIrisEnableRemoveReferenceFallbackPath)
+	{
+		bool bSuccesfullyRemoved = false;
+		for (auto RefHandleIt = ObjectToNetReferenceHandle.CreateIterator(); RefHandleIt; ++RefHandleIt)
+		{
+			if (RefHandleIt.Value() == RefHandle)
 			{
-				if (RefHandleIt.Value() == RefHandle)
-				{
-					RefHandleIt.RemoveCurrent();
-					bSuccesfullyRemoved = true;
-					// It's actually bad if we succesfully remove a reference in this horribly slow fallback path. It means it's still need it and we need to figure out how to get rid of it.
-					UE_LOG_REFERENCECACHE(Error, TEXT("ObjectReferenceCache::RemoveReference: %s. Removed RefHandle from ObjectToNetReferenceHandle."), *RefHandle.ToString());
-					ensureMsgf(false, TEXT("FObjectReferenceCache::RemoveReference %s: Slow fallback path for null instance is still needed. It needs further investigation."), *RefHandle.ToString());
-					break;
-				}
+				RefHandleIt.RemoveCurrent();
+				bSuccesfullyRemoved = true;
+				// It's actually bad if we succesfully remove a reference in this horribly slow fallback path. It means it's still need it and we need to figure out how to get rid of it.
+				UE_LOG_REFERENCECACHE(Error, TEXT("ObjectReferenceCache::RemoveReference: %s. Removed RefHandle from ObjectToNetReferenceHandle."), *RefHandle.ToString());
+				ensureMsgf(false, TEXT("FObjectReferenceCache::RemoveReference %s: Slow fallback path for null instance is still needed. It needs further investigation."), *RefHandle.ToString());
+				break;
 			}
-			UE_CLOG_REFERENCECACHE(!bSuccesfullyRemoved, Verbose, TEXT("ObjectReferenceCache::RemoveReference: %s. Failed removing RefHandle from ObjectToNetReferenceHandle."), *RefHandle.ToString());
 		}
+		UE_CLOG_REFERENCECACHE(!bSuccesfullyRemoved, Verbose, TEXT("ObjectReferenceCache::RemoveReference: %s. Failed removing RefHandle from ObjectToNetReferenceHandle."), *RefHandle.ToString());
 	}
 }
 
@@ -801,7 +798,7 @@ UObject* FObjectReferenceCache::ResolveObjectReferenceHandleInternal(FNetRefHand
 
 		check(IsValid(Object));
 
-		UE_LOG_REFERENCECACHE(VeryVerbose, TEXT("ObjectReferenceCache::ResolveObjectRefererenceHandle Found existing object for %s for Object: %s Pointer: %p"), *RefHandle.ToString(), ToCStr(Object->GetPathName()), Object);
+		UE_LOG_REFERENCECACHE(VeryVerbose, TEXT("ObjectReferenceCache::ResolveObjectRefererenceHandle Found existing object for %s for Object: %s Pointer: 0x%p"), *RefHandle.ToString(), ToCStr(Object->GetPathName()), Object);
 
 		return Object;
 	}
diff --git a/Engine/Source/Runtime/Net/Iris/Private/Iris/ReplicationSystem/ObjectReferenceCache.h b/Engine/Source/Runtime/Net/Iris/Private/Iris/ReplicationSystem/ObjectReferenceCache.h
--- a/Engine/Source/Runtime/Net/Iris/Private/Iris/ReplicationSystem/ObjectReferenceCache.h
+++ b/Engine/Source/Runtime/Net/Iris/Private/Iris/ReplicationSystem/ObjectReferenceCache.h
@@ -87,9 +87,9 @@ public:
 	// Get or create a NetObjectReference from the object identifed by path relative to outer
 	FNetObjectReference GetOrCreateObjectReference(const FString& ObjectPath, const UObject* Outer);
 
-	// Add reference for dynamically spawned object
-	void AddRemoteReference(FNetRefHandle RefHandle, const UObject* Object);
-
+	// Bind a nethandle and the object reference cache on the client
+	void BindRemoteReference(FNetRefHandle RefHandle, const UObject* Object);
+	
 	// Add reference for pre-registered object (client)
 	void AddPreRegisteredReference(FNetRefHandle RefHandle, const UObject* Object);
 
diff --git a/Engine/Source/Runtime/Net/Iris/Private/Iris/ReplicationSystem/ObjectReplicationBridge.cpp b/Engine/Source/Runtime/Net/Iris/Private/Iris/ReplicationSystem/ObjectReplicationBridge.cpp
--- a/Engine/Source/Runtime/Net/Iris/Private/Iris/ReplicationSystem/ObjectReplicationBridge.cpp
+++ b/Engine/Source/Runtime/Net/Iris/Private/Iris/ReplicationSystem/ObjectReplicationBridge.cpp
@@ -980,12 +980,6 @@ void UObjectReplicationBridge::RegisterRemoteInstance(FNetRefHandle RefHandle, U
 	constexpr bool bBindInstanceProtocol = false;
 	InternalAttachInstanceToNetRefHandle(RefHandle, bBindInstanceProtocol, InstanceProtocol, Instance, FNetHandle());
 
-	// Dynamic references needs to be promoted to find the instantiated object
-	if (RefHandle.IsDynamic())
-	{
-		GetObjectReferenceCache()->AddRemoteReference(RefHandle, Instance);
-	}
-
 	UE_LOG_OBJECTREPLICATIONBRIDGE(Verbose, TEXT("RegisterRemoteInstance %s %s with ProtocolId:0x%" UINT64_x_FMT), *PrintObjectFromNetRefHandle(RefHandle), ToCStr(Instance->GetName()), Protocol->ProtocolIdentifier);
 }
 
@@ -1142,6 +1136,12 @@ FReplicationBridgeCreateNetRefHandleResult UObjectReplicationBridge::CreateNetRe
 		// Release instance protocol from the uniquePtr as it is now successfully bound to the handle
 		(void)InstanceProtocol.Release();
 
+		// Dynamic objects or static objects that request it need to update the object cache
+		if (CreateResult.NetRefHandle.IsDynamic() || EnumHasAnyFlags(CreateResult.Flags, EReplicationBridgeCreateNetRefHandleResultFlags::BindStaticObjectInReferenceCache))
+		{
+			GetObjectReferenceCache()->BindRemoteReference(CreateResult.NetRefHandle, InstancePtr);
+		}
+
 		const UNetObjectFactory::FPostInstantiationContext PostContext
 		{ 
 			.Instance = InstancePtr, 
diff --git a/Engine/Source/Runtime/Net/Iris/Public/Iris/ReplicationSystem/ReplicationBridgeTypes.h b/Engine/Source/Runtime/Net/Iris/Public/Iris/ReplicationSystem/ReplicationBridgeTypes.h
--- a/Engine/Source/Runtime/Net/Iris/Public/Iris/ReplicationSystem/ReplicationBridgeTypes.h
+++ b/Engine/Source/Runtime/Net/Iris/Public/Iris/ReplicationSystem/ReplicationBridgeTypes.h
@@ -37,6 +37,8 @@ enum class EReplicationBridgeCreateNetRefHandleResultFlags : uint32
 	AllowDestroyInstanceFromRemote = 1U << 0U,
 	/** Set this flag if you created a subobject and want the RootObject to be notified of the subobject's creation. */
 	ShouldCallSubObjectCreatedFromReplication = AllowDestroyInstanceFromRemote << 1U,
+	/** Bind the static netrefhandle to an object. Needed when the object wasn't found via ResolveObjectReference */
+	BindStaticObjectInReferenceCache = ShouldCallSubObjectCreatedFromReplication << 1U,
 };
 ENUM_CLASS_FLAGS(EReplicationBridgeCreateNetRefHandleResultFlags);
 
