Iris: Expose a net bit array of objects that are filtered out by groups to allow dynamic filters to take group filtering into account.

diff --git a/Engine/Source/Runtime/Net/Iris/Private/Iris/ReplicationSystem/Filtering/ReplicationFiltering.cpp b/Engine/Source/Runtime/Net/Iris/Private/Iris/ReplicationSystem/Filtering/ReplicationFiltering.cpp
--- a/Engine/Source/Runtime/Net/Iris/Private/Iris/ReplicationSystem/Filtering/ReplicationFiltering.cpp
+++ b/Engine/Source/Runtime/Net/Iris/Private/Iris/ReplicationSystem/Filtering/ReplicationFiltering.cpp
@@ -1454,11 +1454,13 @@ void FReplicationFiltering::UpdateDynamicFiltering()
 					continue;
 				}
 
-				FNetObjectFilteringParams FilteringParams;
-				FilteringParams.OutAllowedObjects = AllowedObjects;
-				FilteringParams.FilteringInfos = MakeArrayView(NetObjectFilteringInfos);
-				FilteringParams.ConnectionId = ConnId;
-				FilteringParams.View = Connections->GetReplicationView(ConnId);
+				FNetObjectFilteringParams FilteringParams{
+					.OutAllowedObjects = AllowedObjects,
+					.FilteringInfos = MakeArrayView(NetObjectFilteringInfos),
+					.StateBuffers = nullptr,
+					.ConnectionId = ConnId,
+					.View = Connections->GetReplicationView(ConnId),
+					.GroupFilteredOutObjects = ReplicationSystem->GetReplicationSystemInternal()->GetFiltering().GetGroupFilteredOutObjects(ConnId)};
 
 				// Execute the filter here
 				Info.Filter->Filter(FilteringParams);
diff --git a/Engine/Source/Runtime/Net/Iris/Public/Iris/ReplicationSystem/Filtering/NetObjectFilter.h b/Engine/Source/Runtime/Net/Iris/Public/Iris/ReplicationSystem/Filtering/NetObjectFilter.h
--- a/Engine/Source/Runtime/Net/Iris/Public/Iris/ReplicationSystem/Filtering/NetObjectFilter.h
+++ b/Engine/Source/Runtime/Net/Iris/Public/Iris/ReplicationSystem/Filtering/NetObjectFilter.h
@@ -69,6 +69,9 @@ struct FNetObjectFilteringParams
 
 	/** The view associated with the connection and its sub-connections that objects are filtered for. */
 	UE::Net::FReplicationView View;
+
+	/** List of objects that have been filtered out by groups for the ConnectionId */
+	const UE::Net::FNetBitArrayView GroupFilteredOutObjects;
 };
 
 /**
