Iris
* Added documentation to FilterOut filter

#rnx

diff --git a/Engine/Source/Runtime/Net/Iris/Private/Iris/ReplicationSystem/Filtering/FilterOutNetObjectFilter.cpp b/Engine/Source/Runtime/Net/Iris/Private/Iris/ReplicationSystem/Filtering/FilterOutNetObjectFilter.cpp
--- a/Engine/Source/Runtime/Net/Iris/Private/Iris/ReplicationSystem/Filtering/FilterOutNetObjectFilter.cpp
+++ b/Engine/Source/Runtime/Net/Iris/Private/Iris/ReplicationSystem/Filtering/FilterOutNetObjectFilter.cpp
@@ -4,19 +4,6 @@
 
 #include UE_INLINE_GENERATED_CPP_BY_NAME(FilterOutNetObjectFilter)
 
-void UFilterOutNetObjectFilter::OnInit(const FNetObjectFilterInitParams& Params)
-{
-}
-
-bool UFilterOutNetObjectFilter::AddObject(uint32 ObjectIndex, FNetObjectFilterAddObjectParams& Params)
-{
-	return true;
-}
-
-void UFilterOutNetObjectFilter::RemoveObject(uint32 ObjectIndex, const FNetObjectFilteringInfo& Info)
-{
-}
-
 void UFilterOutNetObjectFilter::Filter(FNetObjectFilteringParams& Params)
 {
 	// Filter out everything
diff --git a/Engine/Source/Runtime/Net/Iris/Public/Iris/ReplicationSystem/Filtering/FilterOutNetObjectFilter.h b/Engine/Source/Runtime/Net/Iris/Public/Iris/ReplicationSystem/Filtering/FilterOutNetObjectFilter.h
--- a/Engine/Source/Runtime/Net/Iris/Public/Iris/ReplicationSystem/Filtering/FilterOutNetObjectFilter.h
+++ b/Engine/Source/Runtime/Net/Iris/Public/Iris/ReplicationSystem/Filtering/FilterOutNetObjectFilter.h
@@ -11,6 +11,11 @@ class UFilterOutNetObjectFilterConfig final : public UNetObjectFilterConfig
 	GENERATED_BODY()
 };
 
+/**
+ * Objects added to this filter will never be relevant automatically.
+ * In order to become relevant to a connection they need to be part of higher-priority relevancy settings.
+ * Ex: @see UReplicationSystem::AddInclusionFilterGroup, UObjectReplicationBridge::AddDependentObject, etc.
+ */
 UCLASS()
 class UFilterOutNetObjectFilter final : public UNetObjectFilter
 {
@@ -18,10 +23,11 @@ class UFilterOutNetObjectFilter final : public UNetObjectFilter
 
 protected:
 	// UNetObjectFilter interface
-	IRISCORE_API virtual void OnInit(const FNetObjectFilterInitParams&) override;
+	virtual void OnInit(const FNetObjectFilterInitParams&) override {}
 	virtual void OnDeinit() override {}
 	virtual void OnMaxInternalNetRefIndexIncreased(uint32 NewMaxInternalIndex) override {}
-	IRISCORE_API virtual bool AddObject(uint32 ObjectIndex, FNetObjectFilterAddObjectParams&) override;
-	IRISCORE_API virtual void RemoveObject(uint32 ObjectIndex, const FNetObjectFilteringInfo&) override;
+	virtual bool AddObject(uint32 ObjectIndex, FNetObjectFilterAddObjectParams&) override { return true; }
+	virtual void RemoveObject(uint32 ObjectIndex, const FNetObjectFilteringInfo&) override {}
+
 	IRISCORE_API virtual void Filter(FNetObjectFilteringParams&) override;
 };
