Fix static validation warnings in IrisCore module for Clang 20

Remove potential to accidentally dereference dangling pointers to previous stack usages.

diff --git a/Engine/Source/Runtime/Net/Iris/Private/Iris/ReplicationSystem/ChunkedDataStream/ChunkedDataReader.cpp b/Engine/Source/Runtime/Net/Iris/Private/Iris/ReplicationSystem/ChunkedDataStream/ChunkedDataReader.cpp
--- a/Engine/Source/Runtime/Net/Iris/Private/Iris/ReplicationSystem/ChunkedDataStream/ChunkedDataReader.cpp
+++ b/Engine/Source/Runtime/Net/Iris/Private/Iris/ReplicationSystem/ChunkedDataStream/ChunkedDataReader.cpp
@@ -481,6 +481,9 @@ void FChunkedDataReader::ReadData(UE::Net::FNetSerializationContext& Context)
 
 	// Assemble data chunks that we have received
 	AssemblePayloadsPendingAssembly(Context);
+
+	// Remove dangling reference to InternalContext on the stack
+	Context.SetInternalContext(nullptr);
 }
 
 void FChunkedDataReader::SetError(const FString& InErrorMessage)
diff --git a/Engine/Source/Runtime/Net/Iris/Private/Iris/ReplicationSystem/ReplicationOperations.cpp b/Engine/Source/Runtime/Net/Iris/Private/Iris/ReplicationSystem/ReplicationOperations.cpp
--- a/Engine/Source/Runtime/Net/Iris/Private/Iris/ReplicationSystem/ReplicationOperations.cpp
+++ b/Engine/Source/Runtime/Net/Iris/Private/Iris/ReplicationSystem/ReplicationOperations.cpp
@@ -1044,6 +1044,9 @@ void FReplicationProtocolOperations::SerializeWithMask(FNetSerializationContext&
 		CurrentInternalStateBuffer += CurrentDescriptor->InternalSize;
 		CurrentChangeMaskBitOffset += CurrentDescriptor->ChangeMaskBitCount;
 	}
+
+	// Remove dangling reference to ChangeMask on the stack
+	Context.SetChangeMask(nullptr);
 }
 
 void FReplicationProtocolOperations::DeserializeWithMask(FNetSerializationContext& Context, uint32* DstChangeMaskData, uint8* RESTRICT DstObjectStateBuffer, const FReplicationProtocol* Protocol)
@@ -1174,6 +1177,9 @@ void FReplicationProtocolOperations::SerializeInitialStateWithMask(FNetSerializa
 		CurrentInternalStateBuffer += CurrentDescriptor->InternalSize;
 		CurrentChangeMaskBitOffset += CurrentDescriptor->ChangeMaskBitCount;
 	}
+
+	// Remove dangling reference to ChangeMask on the stack
+	Context.SetChangeMask(nullptr);
 }
 
 void FReplicationProtocolOperations::DeserializeInitialStateWithMask(FNetSerializationContext& Context, uint32* DstChangeMaskData, uint8* RESTRICT DstObjectStateBuffer, const FReplicationProtocol* Protocol)
@@ -1232,6 +1238,9 @@ void FReplicationProtocolOperations::DeserializeInitialStateWithMask(FNetSeriali
 		CurrentInternalStateBuffer += CurrentDescriptor->InternalSize;
 		CurrentChangeMaskBitOffset += CurrentDescriptor->ChangeMaskBitCount;
 	}
+
+	// Remove dangling reference to ChangeMask on the stack
+	Context.SetChangeMask(nullptr);
 }
 
 void FReplicationProtocolOperations::SerializeWithMaskDelta(FNetSerializationContext& Context, const uint32* ChangeMaskData, const uint8* RESTRICT SrcObjectStateBuffer, const uint8* RESTRICT PrevObjectStateBuffer, const FReplicationProtocol* Protocol)
@@ -1285,6 +1294,9 @@ void FReplicationProtocolOperations::SerializeWithMaskDelta(FNetSerializationCon
 		PrevInternalStateBuffer += CurrentDescriptor->InternalSize;
 		CurrentChangeMaskBitOffset += CurrentDescriptor->ChangeMaskBitCount;
 	}
+
+	// Remove dangling reference to ChangeMask on the stack
+	Context.SetChangeMask(nullptr);
 }
 
 void FReplicationProtocolOperations::DeserializeWithMaskDelta(FNetSerializationContext& Context, uint32* DstChangeMaskData, uint8* RESTRICT DstObjectStateBuffer, const uint8* RESTRICT PrevObjectStateBuffer, const FReplicationProtocol* Protocol)
