Iris - Make memory page protect calls transactionally safe.

diff --git a/Engine/Source/Runtime/Net/Iris/Private/Iris/ReplicationSystem/NetBlob/NetBlob.cpp b/Engine/Source/Runtime/Net/Iris/Private/Iris/ReplicationSystem/NetBlob/NetBlob.cpp
--- a/Engine/Source/Runtime/Net/Iris/Private/Iris/ReplicationSystem/NetBlob/NetBlob.cpp
+++ b/Engine/Source/Runtime/Net/Iris/Private/Iris/ReplicationSystem/NetBlob/NetBlob.cpp
@@ -13,6 +13,7 @@
 #include "Iris/Serialization/ObjectNetSerializer.h"
 #include "HAL/PlatformMemory.h"
 #include "Templates/AlignmentTemplates.h"
+#include "AutoRTFM.h"
 
 namespace UE::Net
 {
@@ -203,9 +204,18 @@ void FNetBlob::FQuantizedBlobState::Protect()
 {
 	if (MemoryAllocationFlags == EMemoryAllocationFlags::Protectable)
 	{
-		constexpr bool bCanRead = true;
-		constexpr bool bCanWrite = false;
-		FPlatformMemory::PageProtect(StateBuffer, AllocationSize, bCanRead, bCanWrite);
+		UE_AUTORTFM_OPEN
+		{
+			constexpr bool bCanRead = true;
+			constexpr bool bCanWrite = false;
+			FPlatformMemory::PageProtect(StateBuffer, AllocationSize, bCanRead, bCanWrite);
+		};
+		UE_AUTORTFM_ONPREABORT(this)
+		{
+			constexpr bool bCanRead = true;
+			constexpr bool bCanWrite = true;
+			FPlatformMemory::PageProtect(StateBuffer, AllocationSize, bCanRead, bCanWrite);
+		};
 	}
 }
 
@@ -213,9 +223,13 @@ void FNetBlob::FQuantizedBlobState::Unprotect()
 {
 	if (MemoryAllocationFlags == EMemoryAllocationFlags::Protectable)
 	{
-		constexpr bool bCanRead = true;
-		constexpr bool bCanWrite = true;
-		FPlatformMemory::PageProtect(StateBuffer, AllocationSize, bCanRead, bCanWrite);
+		// Marking a page as read and write does not require to be rolled back. It's the normal status of our memory.
+		UE_AUTORTFM_OPEN
+		{
+			constexpr bool bCanRead = true;
+			constexpr bool bCanWrite = true;
+			FPlatformMemory::PageProtect(StateBuffer, AllocationSize, bCanRead, bCanWrite);
+		};
 	}
 }
 
