Iris - Don't call FNetBitArray GetSetBitIndices with zero storage.

diff --git a/Engine/Source/Runtime/Net/Iris/Private/Iris/ReplicationSystem/NetBlob/NetBlobManager.cpp b/Engine/Source/Runtime/Net/Iris/Private/Iris/ReplicationSystem/NetBlob/NetBlobManager.cpp
--- a/Engine/Source/Runtime/Net/Iris/Private/Iris/ReplicationSystem/NetBlob/NetBlobManager.cpp
+++ b/Engine/Source/Runtime/Net/Iris/Private/Iris/ReplicationSystem/NetBlob/NetBlobManager.cpp
@@ -527,7 +527,10 @@ void FNetBlobManager::FNetObjectAttachmentSendQueue::PrepareAndProcessOOBAttachm
 		const FNetBitArrayView ReplicatingConnections = MakeNetBitArrayView(ValidConnections);
 
 		ProcessContext.ConnectionIds.SetNum(ReplicatingConnections.CountSetBits());
-		ReplicatingConnections.GetSetBitIndices(0, ~0, ProcessContext.ConnectionIds.GetData(), ProcessContext.ConnectionIds.Num());
+		if (!ProcessContext.ConnectionIds.IsEmpty())
+		{
+			ReplicatingConnections.GetSetBitIndices(0, ~0, ProcessContext.ConnectionIds.GetData(), ProcessContext.ConnectionIds.Num());
+		}
 	}
 
 	uint32 CurrentEntryIndex = 0U;
@@ -576,7 +579,10 @@ void FNetBlobManager::FNetObjectAttachmentSendQueue::PrepareProcessQueue(FReplic
 		const FNetBitArrayView ReplicatingConnections = MakeNetBitArrayView(ValidConnections);
 
 		ProcessContext.ConnectionIds.SetNum(ReplicatingConnections.CountSetBits());
-		ReplicatingConnections.GetSetBitIndices(0, ~0, ProcessContext.ConnectionIds.GetData(), ProcessContext.ConnectionIds.Num());
+		if (!ProcessContext.ConnectionIds.IsEmpty())
+		{
+			ReplicatingConnections.GetSetBitIndices(0, ~0, ProcessContext.ConnectionIds.GetData(), ProcessContext.ConnectionIds.Num());
+		}
 	}
 
 	if (Manager->AllowObjectReplication())
