Iris - Make memory page protect calls transactionally safe.

diff --git a/Engine/Source/Runtime/Net/Iris/Private/Iris/ReplicationSystem/NetBlob/NetBlob.cpp b/Engine/Source/Runtime/Net/Iris/Private/Iris/ReplicationSystem/NetBlob/NetBlob.cpp
--- a/Engine/Source/Runtime/Net/Iris/Private/Iris/ReplicationSystem/NetBlob/NetBlob.cpp
+++ b/Engine/Source/Runtime/Net/Iris/Private/Iris/ReplicationSystem/NetBlob/NetBlob.cpp
@@ -210,7 +210,7 @@ void FNetBlob::FQuantizedBlobState::Protect()
 			constexpr bool bCanWrite = false;
 			FPlatformMemory::PageProtect(StateBuffer, AllocationSize, bCanRead, bCanWrite);
 		};
-		UE_AUTORTFM_ONPREABORT(this)
+		UE_AUTORTFM_ONPREABORT(StateBuffer = this->StateBuffer, AllocationSize = this->AllocationSize)
 		{
 			constexpr bool bCanRead = true;
 			constexpr bool bCanWrite = true;
