Iris
* Added net.Iris.Protocols.IgnoreMismatchForClasses cvar to temporarily disable protocol divergence failures for specific classes.

#rnx

diff --git a/Engine/Source/Runtime/Net/Iris/Private/Iris/ReplicationSystem/ObjectReplicationBridge.cpp b/Engine/Source/Runtime/Net/Iris/Private/Iris/ReplicationSystem/ObjectReplicationBridge.cpp
--- a/Engine/Source/Runtime/Net/Iris/Private/Iris/ReplicationSystem/ObjectReplicationBridge.cpp
+++ b/Engine/Source/Runtime/Net/Iris/Private/Iris/ReplicationSystem/ObjectReplicationBridge.cpp
@@ -98,6 +98,13 @@ static FAutoConsoleVariableRef CVarEnableForceNetUpdate(
 	TEXT("When true the system only allows ForceNetUpdate to skip the poll frequency of objects. When false any MarkDirty object will be immediately polled.")
 	);
 
+static FString IgnoreProtocolMismatchForClasses;
+static FAutoConsoleVariableRef CVarIgnoreProtocolMismatchClasses(
+	TEXT("net.Iris.Protocols.IgnoreMismatchForClasses"),
+	IgnoreProtocolMismatchForClasses,
+	TEXT("Array of class names that when wildcard matched will disable the protocol validation between the server and client. Should be only used temporarely to unblock replication while trying to fix the source of the protocol mismatch.")
+);
+
 #if UE_BUILD_SHIPPING
 static int32 GInvalidLocationWarnings = 0; // don't check OOB locations in Shipping
 #else
@@ -207,6 +214,9 @@ void UObjectReplicationBridge::Initialize(UReplicationSystem* InReplicationSyste
 
 	LoadConfig();
 
+	IgnoreProtocolMismatchForClasses.RemoveSpacesInline();
+	IgnoreProtocolMismatchForClasses.ParseIntoArray(ClassNamesToIgnoreProtocolMismatch, TEXT(","));
+
 	InitConditionalPropertyDelegates();
 	InitNetObjectFactories();
 }
@@ -300,7 +310,7 @@ UE::Net::FNetObjectReference UObjectReplicationBridge::GetOrCreateObjectReferenc
 	return GetObjectReferenceCache()->GetOrCreateObjectReference(Path, Outer);
 }
 
-void UObjectReplicationBridge::AddStaticDestructionInfo(const FString& ObjectPath, const UObject* Outer, const UE::Net::FDestructionParameters& Parameters)
+void UObjectReplicationBridge::AddStaticDestructionInfo(const FString& ObjectPath, const UObject* Outer, const FDestructionParameters& Parameters)
 {
 	UE::Net::FNetObjectReference ObjectRef = GetOrCreateObjectReference(ObjectPath, Outer);
 	if (ObjectRef.IsValid())
@@ -1084,13 +1094,32 @@ FReplicationBridgeCreateNetRefHandleResult UObjectReplicationBridge::CreateNetRe
 	const FReplicationProtocol* ReplicationProtocol = ProtocolManager->GetReplicationProtocol(ReceivedProtocolId, TemplateKey);
 	if (!ReplicationProtocol)
 	{
+		// Make sure the local protocol hash matches the one sent by the server
+		bool bValidateProtocol = true;
+
+		const FString ClassName = InstancePtr->GetClass()->GetName();
+
+		if (ClassNamesToIgnoreProtocolMismatch.Num())
+		{
+			for (const FString& ClassNameToIgnore : ClassNamesToIgnoreProtocolMismatch)
+			{
+				if (ClassName.Contains(ClassNameToIgnore))
+				{
+					// This class wants to keep working even if a protocol mismatch is detected.
+					bValidateProtocol = false;
+					break;
+				}
+			}
+		}
+
+
 		// Create the protocol and compare the received hash ID to make sure it is identical to the one on the server
 		const FCreateReplicationProtocolParameters CreateProtocolParams 
 		{
 			.TemplateKey = TemplateKey, 
-			.bValidateProtocolId = true, // Make sure the local protocol hash matches the one sent by the server
+			.bValidateProtocolId = bValidateProtocol,
 		};
-		ReplicationProtocol = ProtocolManager->CreateReplicationProtocol(ReceivedProtocolId, RegisteredFragments, *(InstancePtr->GetClass()->GetName()), CreateProtocolParams);
+		ReplicationProtocol = ProtocolManager->CreateReplicationProtocol(ReceivedProtocolId, RegisteredFragments, *ClassName, CreateProtocolParams);
 	}
 	else
 	{
diff --git a/Engine/Source/Runtime/Net/Iris/Public/Iris/ReplicationSystem/ObjectReplicationBridge.h b/Engine/Source/Runtime/Net/Iris/Public/Iris/ReplicationSystem/ObjectReplicationBridge.h
--- a/Engine/Source/Runtime/Net/Iris/Public/Iris/ReplicationSystem/ObjectReplicationBridge.h
+++ b/Engine/Source/Runtime/Net/Iris/Public/Iris/ReplicationSystem/ObjectReplicationBridge.h
@@ -42,7 +42,6 @@ namespace UE::Net
 		typedef uint32 FInternalNetRefIndex;
 		class FObjectPollFrequencyLimiter;
 		class FObjectPoller;
-		class FReplicationSystemImpl;
 	}
 }
 
@@ -130,7 +129,7 @@ public:
 	IRISCORE_API FNetRefHandle GetRootObjectOfSubObject(FNetRefHandle SubObjectHandle) const;
 
 	/** Add static destruction info, this is used when stably named objects are destroyed prior to starting replication */
-	IRISCORE_API void AddStaticDestructionInfo(const FString& ObjectPath, const UObject* Outer, const UE::Net::FDestructionParameters& Parameters);
+	IRISCORE_API void AddStaticDestructionInfo(const FString& ObjectPath, const UObject* Outer, const FDestructionParameters& Parameters);
 	
 	/** PIE package name remapping support. */
 	virtual bool RemapPathForPIE(uint32 ConnectionId, FString& Path, bool bReading) const { return false; }
@@ -256,7 +255,7 @@ public:
 		FNetRefHandle InsertRelativeToSubObjectHandle; 
 
 		/** Optional enum to select where in the list this subobject will be inserted */
-		UE::Net::ESubObjectInsertionOrder InsertionOrder = UE::Net::ESubObjectInsertionOrder::None;
+		ESubObjectInsertionOrder InsertionOrder = ESubObjectInsertionOrder::None;
 	};
 
 public:
@@ -474,8 +473,6 @@ private:
 private:
 
 	friend UE::Net::Private::FObjectPoller;
-	friend UReplicationSystem;
-	friend UE::Net::Private::FReplicationSystemImpl;
 
 	TMap<const UClass*, FName> ConfigClassPathNameCache;
 
@@ -567,6 +564,9 @@ private:
 
 	/** Tells if we already logged a warning about a given object or a class type that had a location that is out of bounds. */
 	TSet<FObjectKey> WorldLocationOOBWarnings;
+
+	/** List of class names that will not prevent replication to happen when a protocol mismatch is detected.*/
+	TArray<FString> ClassNamesToIgnoreProtocolMismatch;
 };
 
 ENUM_CLASS_FLAGS(UObjectReplicationBridge::EPrintDebugInfoTraits);
\ No newline at end of file
