Iris
* Fix integration error do to backout cl not merged down into main intentionnaly.

#rnx
#DisabledValidations Preflight

diff --git a/Engine/Source/Runtime/Net/Iris/Private/Iris/ReplicationSystem/ObjectReplicationBridge.cpp b/Engine/Source/Runtime/Net/Iris/Private/Iris/ReplicationSystem/ObjectReplicationBridge.cpp
--- a/Engine/Source/Runtime/Net/Iris/Private/Iris/ReplicationSystem/ObjectReplicationBridge.cpp
+++ b/Engine/Source/Runtime/Net/Iris/Private/Iris/ReplicationSystem/ObjectReplicationBridge.cpp
@@ -310,7 +310,7 @@ UE::Net::FNetObjectReference UObjectReplicationBridge::GetOrCreateObjectReferenc
 	return GetObjectReferenceCache()->GetOrCreateObjectReference(Path, Outer);
 }
 
-void UObjectReplicationBridge::AddStaticDestructionInfo(const FString& ObjectPath, const UObject* Outer, const FDestructionParameters& Parameters)
+void UObjectReplicationBridge::AddStaticDestructionInfo(const FString& ObjectPath, const UObject* Outer, const UE::Net::FDestructionParameters& Parameters)
 {
 	UE::Net::FNetObjectReference ObjectRef = GetOrCreateObjectReference(ObjectPath, Outer);
 	if (ObjectRef.IsValid())
diff --git a/Engine/Source/Runtime/Net/Iris/Public/Iris/ReplicationSystem/ObjectReplicationBridge.h b/Engine/Source/Runtime/Net/Iris/Public/Iris/ReplicationSystem/ObjectReplicationBridge.h
--- a/Engine/Source/Runtime/Net/Iris/Public/Iris/ReplicationSystem/ObjectReplicationBridge.h
+++ b/Engine/Source/Runtime/Net/Iris/Public/Iris/ReplicationSystem/ObjectReplicationBridge.h
@@ -42,6 +42,7 @@ namespace UE::Net
 		typedef uint32 FInternalNetRefIndex;
 		class FObjectPollFrequencyLimiter;
 		class FObjectPoller;
+		class FReplicationSystemImpl;
 	}
 }
 
@@ -129,7 +130,7 @@ public:
 	IRISCORE_API FNetRefHandle GetRootObjectOfSubObject(FNetRefHandle SubObjectHandle) const;
 
 	/** Add static destruction info, this is used when stably named objects are destroyed prior to starting replication */
-	IRISCORE_API void AddStaticDestructionInfo(const FString& ObjectPath, const UObject* Outer, const FDestructionParameters& Parameters);
+	IRISCORE_API void AddStaticDestructionInfo(const FString& ObjectPath, const UObject* Outer, const UE::Net::FDestructionParameters& Parameters);
 	
 	/** PIE package name remapping support. */
 	virtual bool RemapPathForPIE(uint32 ConnectionId, FString& Path, bool bReading) const { return false; }
@@ -255,7 +256,7 @@ public:
 		FNetRefHandle InsertRelativeToSubObjectHandle; 
 
 		/** Optional enum to select where in the list this subobject will be inserted */
-		ESubObjectInsertionOrder InsertionOrder = ESubObjectInsertionOrder::None;
+		UE::Net::ESubObjectInsertionOrder InsertionOrder = UE::Net::ESubObjectInsertionOrder::None;
 	};
 
 public:
@@ -473,6 +474,8 @@ private:
 private:
 
 	friend UE::Net::Private::FObjectPoller;
+	friend UReplicationSystem;
+	friend UE::Net::Private::FReplicationSystemImpl;
 
 	TMap<const UClass*, FName> ConfigClassPathNameCache;
 
