Iris - Always add the split netblob's object references to the first partial blob.

#rnx

diff --git a/Engine/Source/Runtime/Net/Iris/Private/Iris/ReplicationSystem/NetBlob/PartialNetBlob.cpp b/Engine/Source/Runtime/Net/Iris/Private/Iris/ReplicationSystem/NetBlob/PartialNetBlob.cpp
--- a/Engine/Source/Runtime/Net/Iris/Private/Iris/ReplicationSystem/NetBlob/PartialNetBlob.cpp
+++ b/Engine/Source/Runtime/Net/Iris/Private/Iris/ReplicationSystem/NetBlob/PartialNetBlob.cpp
@@ -287,12 +287,19 @@ void FPartialNetBlob::SplitPayload(const FPartialNetBlob::FPayloadSplitParams& S
 		PartialBlob->PartCount = static_cast<uint16>(PartialBlobCount);
 		PartialBlob->SequenceFlags = (bIsFirstPart ? ESequenceFlags::IsFirstPart : ESequenceFlags::None);
 		PartialBlob->SequenceNumber = SequenceNumber++;
-		if (bIsFirstPart && (!SplitParams.ObjectReferencesPendingExport.IsEmpty() || !SplitParams.NetTokensPendingExport.IsEmpty()))
+		if (bIsFirstPart && (EnumHasAnyFlags(SplitParams.OriginalCreationInfo.Flags, ENetBlobFlags::HasExports) || !SplitParams.ObjectReferencesPendingExport.IsEmpty() || !SplitParams.NetTokensPendingExport.IsEmpty()))
 		{
 			PartialBlob->OriginalBlob = SplitParams.OriginalBlob;
 			PartialBlob->CreationInfo.Flags |= ENetBlobFlags::HasExports;
 			PartialBlob->ObjectReferenceExportsArray.Append(SplitParams.ObjectReferencesPendingExport);
 			PartialBlob->NetTokenExportsArray.Append(SplitParams.NetTokensPendingExport);
+			if (EnumHasAnyFlags(SplitParams.OriginalCreationInfo.Flags, ENetBlobFlags::HasExports) && SplitParams.OriginalBlob)
+			{
+				for (TConstArrayView<FNetObjectReference> ObjectReferences = SplitParams.OriginalBlob->CallGetNetObjectReferenceExports(); const FNetObjectReference& ObjectReference : ObjectReferences)
+				{
+					PartialBlob->ObjectReferenceExportsArray.AddUnique(ObjectReference);
+				}
+			}
 		}
 
 		// Copy relevant data from our temporary buffer.
