Iris
- Improved error message if Instanced structtype could not be resolved.

diff --git a/Engine/Source/Runtime/Net/Iris/Private/Iris/Serialization/InstancedStructNetSerializer.cpp b/Engine/Source/Runtime/Net/Iris/Private/Iris/Serialization/InstancedStructNetSerializer.cpp
--- a/Engine/Source/Runtime/Net/Iris/Private/Iris/Serialization/InstancedStructNetSerializer.cpp
+++ b/Engine/Source/Runtime/Net/Iris/Private/Iris/Serialization/InstancedStructNetSerializer.cpp
@@ -207,7 +207,7 @@ void FInstancedStructNetSerializer::Deserialize(FNetSerializationContext& Contex
 		}
 
 		// Deserialize struct data
-		if (ensureMsgf(Struct != nullptr, TEXT("Unable to find struct using FQuantizedObjectReference %s"), ToCStr(StructType.ToString())))
+		if (Struct != nullptr)
 		{
 			FStructNetSerializerConfig StructConfig;
 			StructConfig.StateDescriptor = Config->DescriptorCache.FindOrAddDescriptor(Struct);
@@ -237,6 +237,9 @@ void FInstancedStructNetSerializer::Deserialize(FNetSerializationContext& Contex
 		}
 		else
 		{
+			UE_LOG(LogIris, Error, TEXT("Unable to find struct using FQuantizedObjectReference %s when deserializing instanced struct."), ToCStr(StructType.ToString()));
+			ensure(Struct != nullptr);
+
 			Context.SetError(GNetError_InvalidValue);
 			return;
 		}
