Iris
* Remove private dependencies between ReplicationWriter and ReplicationBridge
* Made the ReplicationWriter hold a pointer to the UObjectReplicationBridge instead

#rnx

diff --git a/Engine/Source/Runtime/Net/Iris/Private/Iris/ReplicationSystem/ObjectReplicationBridge.cpp b/Engine/Source/Runtime/Net/Iris/Private/Iris/ReplicationSystem/ObjectReplicationBridge.cpp
--- a/Engine/Source/Runtime/Net/Iris/Private/Iris/ReplicationSystem/ObjectReplicationBridge.cpp
+++ b/Engine/Source/Runtime/Net/Iris/Private/Iris/ReplicationSystem/ObjectReplicationBridge.cpp
@@ -38,6 +38,7 @@
 #include "Iris/ReplicationSystem/ReplicationFragmentUtil.h"
 #include "Iris/ReplicationSystem/ReplicationOperations.h"
 #include "Iris/ReplicationSystem/ReplicationOperationsInternal.h"
+#include "Iris/ReplicationSystem/ReplicationProtocolManager.h"
 #include "Iris/ReplicationSystem/ReplicationSystem.h"
 #include "Iris/ReplicationSystem/ReplicationSystemInternal.h"
 #include "Iris/ReplicationSystem/ReplicationSystemTypes.h"
@@ -205,6 +206,8 @@ void UObjectReplicationBridge::Initialize(UReplicationSystem* InReplicationSyste
 {
 	Super::Initialize(InReplicationSystem);
 
+	CreateDestructionProtocol();
+
 	const uint32 CurrentMaxInternalIndex = NetRefHandleManager->GetCurrentMaxInternalNetRefIndex();
 	PollFrequencyLimiter->Init(CurrentMaxInternalIndex);
 	ObjectsWithObjectReferences.Init(CurrentMaxInternalIndex);
@@ -232,9 +235,36 @@ void UObjectReplicationBridge::Deinitialize()
 	OnCustomConditionChangedHandle.Reset();
 	OnDynamicConditionChangedHandle.Reset();
 	PollFrequencyLimiter->Deinit();
+
+	// Release destructioninfo protocol.
+	if (DestructionInfoProtocol)
+	{
+		DestructionInfoProtocol->Release();
+		DestructionInfoProtocol = nullptr;
+	}
+
 	Super::Deinitialize();
 }
 
+void UObjectReplicationBridge::CreateDestructionProtocol()
+{
+	using namespace UE::Net;
+	using namespace UE::Net::Private;
+
+	// Create destruction info protocol
+	const FReplicationFragments RegisteredFragments;
+	const FReplicationProtocolIdentifier ProtocolId = FReplicationProtocolManager::CalculateProtocolIdentifier(RegisteredFragments);
+	const FCreateReplicationProtocolParameters CreateProtocolParams
+	{
+		.bValidateProtocolId = false,
+		.bHasTemplateKey = false, // Destruction info protocol doesn't need a TemplateKey
+		.TypeStatsIndex = GetReplicationSystem()->GetReplicationSystemInternal()->GetNetTypeStats().GetOrCreateTypeStats(FName("DestructionInfo")),
+	};
+	DestructionInfoProtocol = ReplicationProtocolManager->CreateReplicationProtocol(ProtocolId, RegisteredFragments, TEXT("InternalDestructionInfo"), CreateProtocolParams);
+	// Explicit refcount
+	DestructionInfoProtocol->AddRef();
+}
+
 void UObjectReplicationBridge::InitNetObjectFactories()
 {
 	using namespace UE::Net;
@@ -319,6 +349,134 @@ void UObjectReplicationBridge::AddStaticDestructionInfo(const FString& ObjectPat
 	}
 }
 
+UE::Net::FNetRefHandle UObjectReplicationBridge::StoreDestructionInfo(FNetRefHandle Handle, const UE::Net::FDestructionParameters& Parameters)
+{
+	using namespace UE::Net;
+	using namespace UE::Net::Private;
+
+	if (!FNetObjectFactoryRegistry::IsValidFactoryId(Parameters.NetFactoryId))
+	{
+		ensureMsgf(false, TEXT("StoreDestructionInfo received invalid NetFactoryId: %u for %s"), Parameters.NetFactoryId, *PrintObjectFromNetRefHandle(Handle));
+		UE_LOG(LogIris, Error, TEXT("StoreDestructionInfo received invalid NetFactoryID: %u for %s"), Parameters.NetFactoryId, *PrintObjectFromNetRefHandle(Handle));
+
+		return FNetRefHandle::GetInvalid();
+	}
+
+	// Not allowed to create destruction infos at the moment
+	if (!CanCreateDestructionInfo())
+	{
+		return FNetRefHandle::GetInvalid();
+	}
+
+	FReplicationProtocolManager* ProtocolManager = GetReplicationProtocolManager();
+
+	const FCreateNetObjectParams CreateNetObjectParams
+	{
+		.NetFactoryId = Parameters.NetFactoryId,
+		.ReplicationProtocol = DestructionInfoProtocol,
+	};
+
+	// Create handle for DestructionInfo to be able to scope destruction infos
+	const FNetObjectGroupHandle LevelGroupHandle = GetLevelGroup(Parameters.Level);
+	FNetRefHandle DestructionInfoHandle = NetRefHandleManager->CreateHandleForDestructionInfo(Handle, CreateNetObjectParams);
+
+	if (LevelGroupHandle.IsValid())
+	{
+		GetReplicationSystem()->AddToGroup(LevelGroupHandle, DestructionInfoHandle);
+	}
+
+	const FInternalNetRefIndex InternalReplicationIndex = NetRefHandleManager->GetInternalIndex(DestructionInfoHandle);
+
+	// We also need to store the actual data we send to destroy static references when they are scoped
+	FStaticDestructionInfo PersistentDestructionInfo
+	{
+		.StaticRef = FObjectReferenceCache::MakeNetObjectReference(Handle),
+		.LevelGroupHandle = LevelGroupHandle,
+		.NetFactoryId = Parameters.NetFactoryId,
+		.InternalReplicationIndex = InternalReplicationIndex
+	};
+
+	StaticObjectsPendingDestroy.Add(DestructionInfoHandle, PersistentDestructionInfo);
+
+	// If we use distance based Prioritization for destruction infos we need to populate the quantized state used for prioritization.
+	const bool bUseDynamicPrioritization = Parameters.bUseDistanceBasedPrioritization;
+	if (bUseDynamicPrioritization)
+	{
+		// Use WorldLocations to feed the location of the destruction info so that it can be prioritized properly.
+		FWorldLocations& WorldLocations = GetReplicationSystem()->GetReplicationSystemInternal()->GetWorldLocations();
+
+		// Check the position lies within world boundaries.
+		ensureMsgf(Parameters.Location.X >= WorldLocations.GetWorldMinPos().X && Parameters.Location.X <= WorldLocations.GetWorldMaxPos().X && 
+				   Parameters.Location.Y >= WorldLocations.GetWorldMinPos().Y && Parameters.Location.Y <= WorldLocations.GetWorldMaxPos().Y, TEXT("Object %s with position %s lies outside configured world boundary."), ToCStr(NetRefHandleManager->PrintObjectFromIndex((InternalReplicationIndex))), *Parameters.Location.ToString());
+
+		WorldLocations.InitObjectInfoCache(InternalReplicationIndex);
+		WorldLocations.SetObjectInfo(InternalReplicationIndex, Parameters.Location, 0.0f);
+
+		GetReplicationSystem()->SetPrioritizer(DestructionInfoHandle, DefaultSpatialNetObjectPrioritizerHandle);
+	}
+
+	UE_LOG_OBJECTREPLICATIONBRIDGE(Verbose, TEXT("StoreDestructionInfo %s (InternalIndex: %u ) for %s GroupIndex: %u"), *DestructionInfoHandle.ToString(), InternalReplicationIndex,  *PrintObjectFromNetRefHandle(Handle), LevelGroupHandle.GetGroupIndex());
+	
+	return DestructionInfoHandle;
+}
+
+void UObjectReplicationBridge::RemoveDestructionInfosForGroup(UE::Net::FNetObjectGroupHandle GroupHandle)
+{
+	using namespace UE::Net;
+	using namespace UE::Net::Private;
+
+	FWorldLocations& WorldLocations = GetReplicationSystem()->GetReplicationSystemInternal()->GetWorldLocations();
+
+	if (GroupHandle.IsValid())
+	{
+		UE_LOG_OBJECTREPLICATIONBRIDGE(Verbose, TEXT("RemoveDestructionInfosForGroup GroupIndex: %u"), GroupHandle.GetGroupIndex());
+
+		FNetObjectGroup* Group = Groups->GetGroup(GroupHandle);
+		check(Group);
+
+		TArray<FNetRefHandle, TInlineAllocator<384>> ObjectsToRemove;	
+		TArray<FInternalNetRefIndex, TInlineAllocator<384>> ObjectIndicesToRemove;	
+		for (uint32 InternalObjectIndex : MakeArrayView(Group->Members))
+		{
+			if (NetRefHandleManager->GetIsDestroyedStartupObject(InternalObjectIndex))
+			{
+				const FNetRefHandleManager::FReplicatedObjectData& ObjectData = NetRefHandleManager->GetReplicatedObjectDataNoCheck(InternalObjectIndex);
+				if (StaticObjectsPendingDestroy.Remove(ObjectData.RefHandle))
+				{
+					ObjectsToRemove.Add(ObjectData.RefHandle);
+					ObjectIndicesToRemove.Add(InternalObjectIndex);
+				}
+			}			
+		}
+
+		for (FNetRefHandle Handle : MakeArrayView(ObjectsToRemove))
+		{
+			NetRefHandleManager->DestroyNetObject(Handle);
+		}
+
+		for (FInternalNetRefIndex InternalReplicationIndex : MakeArrayView(ObjectIndicesToRemove))
+		{
+			WorldLocations.RemoveObjectInfoCache(InternalReplicationIndex);
+		}
+	}
+	else
+	{
+		// We should remove all destruction infos and objects
+		for (const auto& It : StaticObjectsPendingDestroy)
+		{
+			NetRefHandleManager->DestroyNetObject(It.Key);
+		}
+
+		// Remove from WorldLocations
+		for (const auto& It : StaticObjectsPendingDestroy)
+		{
+			WorldLocations.RemoveObjectInfoCache(It.Value.InternalReplicationIndex);
+		}
+
+		StaticObjectsPendingDestroy.Empty();
+	}
+}
+
 UObject* UObjectReplicationBridge::GetReplicatedObject(FNetRefHandle Handle) const
 {
 	return IsReplicatedHandle(Handle) ? GetObjectFromReferenceHandle(Handle) : nullptr;
@@ -668,7 +826,7 @@ UE::Net::FNetRefHandle UObjectReplicationBridge::StartReplicatingSubObject(UObje
 	FNetRefHandle PotentiallyPreRegisteredHandle = GetObjectReferenceCache()->GetObjectReferenceHandleFromObject(Instance);
 	if (PotentiallyPreRegisteredHandle.IsValid() && IsNetRefHandlePreRegistered(PotentiallyPreRegisteredHandle))
 	{
-		UE_LOG_OBJECTREPLICATIONBRIDGE(Error, TEXT("UObjectReplicationBridge::StartReplicatingSubObject attempted to replicate pre-registered object %s."), *GetFullNameSafe(Instance));
+		UE_LOG_OBJECTREPLICATIONBRIDGE(Error, TEXT("StartReplicatingSubObject attempted to replicate pre-registered object %s."), *GetFullNameSafe(Instance));
 		ensureMsgf(false, TEXT("Attempted to replicate a pre-registered subobject"));
 		return FNetRefHandle::GetInvalid();
 	}
@@ -680,7 +838,7 @@ UE::Net::FNetRefHandle UObjectReplicationBridge::StartReplicatingSubObject(UObje
 		const FNetRefHandle ExistingRootObjectHandle = LocalNetRefHandleManager.GetRootObjectOfSubObject(SubObjectRefHandle);
 		if (Params.RootObjectHandle != ExistingRootObjectHandle)
 		{
-			UE_LOG_OBJECTREPLICATIONBRIDGE(Error, TEXT("UObjectReplicationBridge::StartReplicatingSubObject attempted to replicate already replicating subobject (%s) with a different rootobject (%s)."),
+			UE_LOG_OBJECTREPLICATIONBRIDGE(Error, TEXT("StartReplicatingSubObject attempted to replicate already replicating subobject (%s) with a different rootobject (%s)."),
 				*LocalNetRefHandleManager.PrintObjectFromNetRefHandle(SubObjectRefHandle), *LocalNetRefHandleManager.PrintObjectFromNetRefHandle(Params.RootObjectHandle));
 			ensureMsgf(false, TEXT("Attempted to replicate already replicating subobject with a different root"));
 		}
@@ -831,6 +989,8 @@ bool UObjectReplicationBridge::WriteNetRefHandleCreationInfo(FReplicationBridgeS
 	using namespace UE::Net;
 	using namespace UE::Net::Private;
 
+	check(!Context.bIsDestructionInfo);
+
 	// If there is a cached header for this handle, use it!
 	if (TUniquePtr<const FNetObjectCreationHeader>* CachedHeader = CachedCreationHeaders.Find(Handle))
 	{
diff --git a/Engine/Source/Runtime/Net/Iris/Private/Iris/ReplicationSystem/ReplicationBridge.cpp b/Engine/Source/Runtime/Net/Iris/Private/Iris/ReplicationSystem/ReplicationBridge.cpp
--- a/Engine/Source/Runtime/Net/Iris/Private/Iris/ReplicationSystem/ReplicationBridge.cpp
+++ b/Engine/Source/Runtime/Net/Iris/Private/Iris/ReplicationSystem/ReplicationBridge.cpp
@@ -71,7 +71,6 @@ UReplicationBridge::UReplicationBridge()
 , ReplicationProtocolManager(nullptr)
 , ReplicationStateDescriptorRegistry(nullptr)
 , NetRefHandleManager(nullptr)
-, DestructionInfoProtocol(nullptr)
 {
 }
 
@@ -96,11 +95,6 @@ void UReplicationBridge::PostReceiveUpdate()
 	OnPostReceiveUpdate();
 }
 
-bool UReplicationBridge::WriteNetRefHandleCreationInfo(FReplicationBridgeSerializationContext& Context, FNetRefHandle Handle)
-{
-	return true;
-}
-
 bool UReplicationBridge::CacheNetRefHandleCreationInfo(FNetRefHandle Handle)
 {
 	return false;
@@ -256,39 +250,6 @@ bool UReplicationBridge::CallCacheNetRefHandleCreationInfo(FNetRefHandle Handle)
 	return CacheNetRefHandleCreationInfo(Handle);
 }
 
-bool UReplicationBridge::CallWriteNetRefHandleCreationInfo(FReplicationBridgeSerializationContext& Context, FNetRefHandle Handle)
-{
-	using namespace UE::Net;
-
-	check(!Context.bIsDestructionInfo);
-
-	return WriteNetRefHandleCreationInfo(Context, Handle);
-}
-
-bool UReplicationBridge::CallWriteNetRefHandleDestructionInfo(FReplicationBridgeSerializationContext& Context, FNetRefHandle Handle)
-{
-	using namespace UE::Net;
-
-	check(Context.bIsDestructionInfo);
-	UE_NET_TRACE_SCOPE(DestructionInfo, *Context.SerializationContext.GetBitStreamWriter(), Context.SerializationContext.GetTraceCollector(), ENetTraceVerbosity::Trace);
-
-	const FDestructionInfo* Info = StaticObjectsPendingDestroy.Find(Handle);
-	if (ensure(Info))
-	{
-		UE_LOG(LogIris, VeryVerbose, TEXT("CallWriteNetRefHandleDestructionInfo on %s | %s | NetFactoryId: %u"), *PrintObjectFromNetRefHandle(Handle), *Info->StaticRef.ToString(), Info->NetFactoryId);
-		// Write destruction info
-		WriteFullNetObjectReference(Context.SerializationContext, Info->StaticRef);
-		Context.SerializationContext.GetBitStreamWriter()->WriteBits(Info->NetFactoryId, FNetObjectFactoryRegistry::GetMaxBits());
-	}
-	else
-	{
-		Context.SerializationContext.SetError(TEXT("DestructionInfoNotFound"), false);
-		UE_LOG_REPLICATIONBRIDGE(Error, TEXT("Failed to write destructionInfo for %s"), *Handle.ToString());
-	}
-
-	return !Context.SerializationContext.HasErrorOrOverflow();
-}
-
 FReplicationBridgeCreateNetRefHandleResult UReplicationBridge::CallCreateNetRefHandleFromRemote(FNetRefHandle RootObjectOfSubObject, FNetRefHandle WantedNetHandle, FReplicationBridgeSerializationContext& Context)
 {
 	check(!Context.bIsDestructionInfo);
@@ -378,21 +339,6 @@ void UReplicationBridge::Initialize(UReplicationSystem* InReplicationSystem)
 	NetRefHandleManager = &ReplicationSystemInternal->GetNetRefHandleManager();
 	ObjectReferenceCache = &ReplicationSystemInternal->GetObjectReferenceCache();
 	Groups = &ReplicationSystemInternal->GetGroups();
-
-	// Create destruction info protocol
-	{
-		const FReplicationFragments RegisteredFragments;
-		const FReplicationProtocolIdentifier ProtocolId = FReplicationProtocolManager::CalculateProtocolIdentifier(RegisteredFragments);
-		const FCreateReplicationProtocolParameters CreateProtocolParams
-		{
-			.bValidateProtocolId = false,
-			.bHasTemplateKey = false, // Destruction info protocol doesn't need a TemplateKey
-			.TypeStatsIndex =  GetReplicationSystem()->GetReplicationSystemInternal()->GetNetTypeStats().GetOrCreateTypeStats(FName("DestructionInfo")),
-		};
-		DestructionInfoProtocol = ReplicationProtocolManager->CreateReplicationProtocol(ProtocolId, RegisteredFragments, TEXT("InternalDestructionInfo"), CreateProtocolParams);
-		// Explicit refcount
-		DestructionInfoProtocol->AddRef();
-	}
 }
 
 void UReplicationBridge::Deinitialize()
@@ -425,13 +371,6 @@ void UReplicationBridge::Deinitialize()
 		}
 	});
 
-	// Release destructioninfo protocol.
-	if (DestructionInfoProtocol)
-	{
-		DestructionInfoProtocol->Release();
-		DestructionInfoProtocol = nullptr;
-	}
-
 	ReplicationSystem = nullptr;
 	ReplicationProtocolManager = nullptr;
 	ReplicationStateDescriptorRegistry = nullptr;
@@ -732,62 +671,7 @@ void UReplicationBridge::StopReplicatingNetRefHandle(FNetRefHandle Handle, EEndR
 	}
 }
 
-void UReplicationBridge::RemoveDestructionInfosForGroup(UE::Net::FNetObjectGroupHandle GroupHandle)
-{
-	using namespace UE::Net;
-	using namespace UE::Net::Private;
-
-	FWorldLocations& WorldLocations = GetReplicationSystem()->GetReplicationSystemInternal()->GetWorldLocations();
-
-	if (GroupHandle.IsValid())
-	{
-		UE_LOG_REPLICATIONBRIDGE(Verbose, TEXT("RemoveDestructionInfosForGroup GroupIndex: %u"), GroupHandle.GetGroupIndex());
-
-		FNetObjectGroup* Group = Groups->GetGroup(GroupHandle);
-		check(Group);
-
-		TArray<FNetRefHandle, TInlineAllocator<384>> ObjectsToRemove;	
-		TArray<FInternalNetRefIndex, TInlineAllocator<384>> ObjectIndicesToRemove;	
-		for (uint32 InternalObjectIndex : MakeArrayView(Group->Members))
-		{
-			if (NetRefHandleManager->GetIsDestroyedStartupObject(InternalObjectIndex))
-			{
-				const FNetRefHandleManager::FReplicatedObjectData& ObjectData = NetRefHandleManager->GetReplicatedObjectDataNoCheck(InternalObjectIndex);
-				if (StaticObjectsPendingDestroy.Remove(ObjectData.RefHandle))
-				{
-					ObjectsToRemove.Add(ObjectData.RefHandle);
-					ObjectIndicesToRemove.Add(InternalObjectIndex);
-				}
-			}			
-		}
-
-		for (FNetRefHandle Handle : MakeArrayView(ObjectsToRemove))
-		{
-			NetRefHandleManager->DestroyNetObject(Handle);
-		}
-
-		for (FInternalNetRefIndex InternalReplicationIndex : MakeArrayView(ObjectIndicesToRemove))
-		{
-			WorldLocations.RemoveObjectInfoCache(InternalReplicationIndex);
-		}
-	}
-	else
-	{
-		// We should remove all destruction infos and objects
-		for (const auto& It : StaticObjectsPendingDestroy)
-		{
-			NetRefHandleManager->DestroyNetObject(It.Key);
-		}
-
-		// Remove from WorldLocations
-		for (const auto& It : StaticObjectsPendingDestroy)
-		{
-			WorldLocations.RemoveObjectInfoCache(It.Value.InternalReplicationIndex);
-		}
 
-		StaticObjectsPendingDestroy.Empty();
-	}
-}
 
 void UReplicationBridge::TearOffHandlesPendingTearOff()
 {
@@ -1029,77 +913,6 @@ void UReplicationBridge::InternalTearOff(FNetRefHandle Handle)
 	ObjectData.bShouldPropagateChangedStates = 0U;
 }
 
-UE::Net::FNetRefHandle UReplicationBridge::StoreDestructionInfo(FNetRefHandle Handle, const UE::Net::FDestructionParameters& Parameters)
-{
-	using namespace UE::Net;
-	using namespace UE::Net::Private;
-
-	if (!FNetObjectFactoryRegistry::IsValidFactoryId(Parameters.NetFactoryId))
-	{
-		ensureMsgf(false, TEXT("StoreDestructionInfo received invalid NetFactoryId: %u for %s"), Parameters.NetFactoryId, *PrintObjectFromNetRefHandle(Handle));
-		UE_LOG(LogIris, Error, TEXT("StoreDestructionInfo received invalid NetFactoryID: %u for %s"), Parameters.NetFactoryId, *PrintObjectFromNetRefHandle(Handle));
-
-		return FNetRefHandle::GetInvalid();
-	}
-
-	// Not allowed to create destruction infos at the moment
-	if (!CanCreateDestructionInfo())
-	{
-		return FNetRefHandle::GetInvalid();
-	}
-
-	FReplicationProtocolManager* ProtocolManager = GetReplicationProtocolManager();
-
-	const FCreateNetObjectParams CreateNetObjectParams
-	{
-		.NetFactoryId = Parameters.NetFactoryId,
-		.ReplicationProtocol = DestructionInfoProtocol,
-	};
-
-	// Create handle for DestructionInfo to be able to scope destruction infos
-	const FNetObjectGroupHandle LevelGroupHandle = GetLevelGroup(Parameters.Level);
-	FNetRefHandle DestructionInfoHandle = NetRefHandleManager->CreateHandleForDestructionInfo(Handle, CreateNetObjectParams);
-
-	if (LevelGroupHandle.IsValid())
-	{
-		GetReplicationSystem()->AddToGroup(LevelGroupHandle, DestructionInfoHandle);
-	}
-
-	const FInternalNetRefIndex InternalReplicationIndex = NetRefHandleManager->GetInternalIndex(DestructionInfoHandle);
-
-	// We also need to store the actual data we send to destroy static references when they are scoped
-	FDestructionInfo PersistentDestructionInfo
-	{
-		.StaticRef = FObjectReferenceCache::MakeNetObjectReference(Handle),
-		.LevelGroupHandle = LevelGroupHandle,
-		.NetFactoryId = Parameters.NetFactoryId,
-		.InternalReplicationIndex = InternalReplicationIndex
-	};
-
-	StaticObjectsPendingDestroy.Add(DestructionInfoHandle, PersistentDestructionInfo);
-
-	// If we use distance based Prioritization for destruction infos we need to populate the quantized state used for prioritization.
-	const bool bUseDynamicPrioritization = Parameters.bUseDistanceBasedPrioritization;
-	if (bUseDynamicPrioritization)
-	{
-		// Use WorldLocations to feed the location of the destruction info so that it can be prioritized properly.
-		FWorldLocations& WorldLocations = GetReplicationSystem()->GetReplicationSystemInternal()->GetWorldLocations();
-
-		// Check the position lies within world boundaries.
-		ensureMsgf(Parameters.Location.X >= WorldLocations.GetWorldMinPos().X && Parameters.Location.X <= WorldLocations.GetWorldMaxPos().X && 
-				   Parameters.Location.Y >= WorldLocations.GetWorldMinPos().Y && Parameters.Location.Y <= WorldLocations.GetWorldMaxPos().Y, TEXT("Object %s with position %s lies outside configured world boundary."), ToCStr(NetRefHandleManager->PrintObjectFromIndex((InternalReplicationIndex))), *Parameters.Location.ToString());
-
-		WorldLocations.InitObjectInfoCache(InternalReplicationIndex);
-		WorldLocations.SetObjectInfo(InternalReplicationIndex, Parameters.Location, 0.0f);
-
-		GetReplicationSystem()->SetPrioritizer(DestructionInfoHandle, DefaultSpatialNetObjectPrioritizerHandle);
-	}
-
-	UE_LOG_REPLICATIONBRIDGE(Verbose, TEXT("UReplicationBridge::StoreDestructionInfo %s (InternalIndex: %u ) for %s GroupIndex: %u"), *DestructionInfoHandle.ToString(), InternalReplicationIndex,  *PrintObjectFromNetRefHandle(Handle), LevelGroupHandle.GetGroupIndex());
-	
-	return DestructionInfoHandle;
-}
-
 bool UReplicationBridge::IsReplicatedHandle(FNetRefHandle Handle) const
 {
 	return Handle.IsValid() && NetRefHandleManager->IsNetRefHandleAssigned(Handle);
diff --git a/Engine/Source/Runtime/Net/Iris/Private/Iris/ReplicationSystem/ReplicationWriter.cpp b/Engine/Source/Runtime/Net/Iris/Private/Iris/ReplicationSystem/ReplicationWriter.cpp
--- a/Engine/Source/Runtime/Net/Iris/Private/Iris/ReplicationSystem/ReplicationWriter.cpp
+++ b/Engine/Source/Runtime/Net/Iris/Private/Iris/ReplicationSystem/ReplicationWriter.cpp
@@ -2580,11 +2580,11 @@ FReplicationWriter::EWriteObjectStatus FReplicationWriter::WriteObjectAndSubObje
 					bool bWriteSuccess = false;
 					if (Info.IsDestructionInfo)
 					{
-						bWriteSuccess = ReplicationBridge->CallWriteNetRefHandleDestructionInfo(BridgeContext, NetRefHandle);
+						bWriteSuccess = WriteNetRefHandleDestructionInfo(Context, NetRefHandle);
 					}
 					else
 					{
-						bWriteSuccess = ReplicationBridge->CallWriteNetRefHandleCreationInfo(BridgeContext, NetRefHandle);
+						bWriteSuccess = FBridgeSerialization::WriteNetRefHandleCreationInfo(ReplicationBridge, BridgeContext, NetRefHandle);
 					}
 
 					// We need to send creation info, so if we fail we skip this object for now
@@ -3260,7 +3260,7 @@ FReplicationWriter::EWriteStatus FReplicationWriter::WriteDestructionInfo(FNetSe
 
 	// Push ForceInlineExportScope to inline exports instead of writing exports later.
 	FForceInlineExportScope ForceInlineExportScope(Context.GetInternalContext());
-	if (!ReplicationBridge->CallWriteNetRefHandleDestructionInfo(BridgeContext, ObjectData.RefHandle))
+	if (!WriteNetRefHandleDestructionInfo(Context, ObjectData.RefHandle))
 	{
 		// Trigger Rollback
 		Writer.DoOverflow();
@@ -3313,6 +3313,27 @@ FReplicationWriter::EWriteStatus FReplicationWriter::WriteDestructionInfo(FNetSe
 	return Writer.IsOverflown() ? EWriteStatus::Abort : EWriteStatus::Written;
 }
 
+bool FReplicationWriter::WriteNetRefHandleDestructionInfo(FNetSerializationContext& Context, FNetRefHandle Handle)
+{
+	UE_NET_TRACE_SCOPE(DestructionInfo, *Context.GetBitStreamWriter(), Context.GetTraceCollector(), ENetTraceVerbosity::Trace);
+
+	const FStaticDestructionInfo* Info = FBridgeSerialization::GetStaticDestructionInfo(ReplicationBridge, Handle);
+	if (ensure(Info))
+	{
+		UE_LOG(LogIris, VeryVerbose, TEXT("WriteNetRefHandleDestructionInfo on %s | %s | NetFactoryId: %u"), *Handle.ToString(), *Info->StaticRef.ToString(), Info->NetFactoryId);
+		// Write destruction info
+		WriteFullNetObjectReference(Context, Info->StaticRef);
+		Context.GetBitStreamWriter()->WriteBits(Info->NetFactoryId, FNetObjectFactoryRegistry::GetMaxBits());
+	}
+	else
+	{
+		Context.SetError(TEXT("DestructionInfoNotFound"), false);
+		UE_LOG(LogIris, Error, TEXT("Failed to write destructionInfo for %s"), *Handle.ToString());
+	}
+
+	return !Context.HasErrorOrOverflow();
+}
+
 uint32 FReplicationWriter::WriteOOBAttachments(FNetSerializationContext& Context)
 {
 	uint32 WrittenObjectCount = 0U;
@@ -4438,4 +4459,4 @@ void FReplicationWriter::FHugeObjectSendQueue::FreeContexts(TFunctionRef<void (c
 	RootObjectsInTransit.Empty();
 }
 
-}
+} // end namespace UE::Net::Private
diff --git a/Engine/Source/Runtime/Net/Iris/Private/Iris/ReplicationSystem/ReplicationWriter.h b/Engine/Source/Runtime/Net/Iris/Private/Iris/ReplicationSystem/ReplicationWriter.h
--- a/Engine/Source/Runtime/Net/Iris/Private/Iris/ReplicationSystem/ReplicationWriter.h
+++ b/Engine/Source/Runtime/Net/Iris/Private/Iris/ReplicationSystem/ReplicationWriter.h
@@ -25,7 +25,7 @@
 class UReplicationSystem;
 class UNetObjectBlobHandler;
 class UPartialNetObjectAttachmentHandler;
-class UReplicationBridge;
+class UObjectReplicationBridge;
 
 namespace UE::Net
 {
@@ -513,6 +513,8 @@ private:
 	// returns > 1 if Objects was written, 0 if the objects was skipped (failed dependency or waiting for creation confirmation) -1 if the stream is full and we should stop
 	EWriteStatus WriteDestructionInfo(FNetSerializationContext& Context, uint32 InternalIndex);
 
+	bool WriteNetRefHandleDestructionInfo(FNetSerializationContext& Context, FNetRefHandle Handle);
+
 	struct FWriteBatchResult
 	{
 		EWriteStatus Status = EWriteStatus::Skipped;
@@ -628,7 +630,7 @@ private:
 	// Cached internal systems
 	FReplicationSystemInternal* ReplicationSystemInternal = nullptr;
 	FNetRefHandleManager* NetRefHandleManager = nullptr;
-	UReplicationBridge* ReplicationBridge = nullptr;
+	UObjectReplicationBridge* ReplicationBridge = nullptr;
 	FDeltaCompressionBaselineManager* BaselineManager = nullptr;
 	FObjectReferenceCache* ObjectReferenceCache = nullptr;
 	const FReplicationFiltering* ReplicationFiltering = nullptr;
diff --git a/Engine/Source/Runtime/Net/Iris/Public/Iris/ReplicationSystem/ObjectReplicationBridge.h b/Engine/Source/Runtime/Net/Iris/Public/Iris/ReplicationSystem/ObjectReplicationBridge.h
--- a/Engine/Source/Runtime/Net/Iris/Public/Iris/ReplicationSystem/ObjectReplicationBridge.h
+++ b/Engine/Source/Runtime/Net/Iris/Public/Iris/ReplicationSystem/ObjectReplicationBridge.h
@@ -43,9 +43,31 @@ namespace UE::Net
 		class FObjectPollFrequencyLimiter;
 		class FObjectPoller;
 		class FReplicationSystemImpl;
+		class FBridgeSerialization;
 	}
 }
 
+namespace UE::Net::Private
+{
+
+/** Info representing destroyed static objects */
+struct FStaticDestructionInfo
+{
+	/** The reference to the static object that is considered destroyed */
+	UE::Net::FNetObjectReference StaticRef;
+	
+	/** The level the object was a part of */
+	UE::Net::FNetObjectGroupHandle LevelGroupHandle;
+	
+	/** The NetFactory assigned to this object */
+	UE::Net::FNetObjectFactoryId NetFactoryId = UE::Net::InvalidNetObjectFactoryId;
+
+	/** The repindex assigned to this destruction info */
+	UE::Net::Private::FInternalNetRefIndex InternalReplicationIndex;
+};
+
+} // end namespace UE::Net::Private
+
 struct FObjectReplicationBridgeInstantiateResult
 {
 	UObject* Object = nullptr;
@@ -131,6 +153,9 @@ public:
 
 	/** Add static destruction info, this is used when stably named objects are destroyed prior to starting replication */
 	IRISCORE_API void AddStaticDestructionInfo(const FString& ObjectPath, const UObject* Outer, const UE::Net::FDestructionParameters& Parameters);
+
+	/** Store static destruction info for the replicated object. */
+	IRISCORE_API FNetRefHandle StoreDestructionInfo(FNetRefHandle Handle, const UE::Net::FDestructionParameters& Parameters);
 	
 	/** PIE package name remapping support. */
 	virtual bool RemapPathForPIE(uint32 ConnectionId, FString& Path, bool bReading) const { return false; }
@@ -297,7 +322,6 @@ protected:
 	IRISCORE_API virtual void OnPostReceiveUpdate() override;
 	IRISCORE_API virtual void UpdateInstancesWorldLocation() override;
 	IRISCORE_API virtual void PruneStaleObjects() override;	
-	IRISCORE_API virtual bool WriteNetRefHandleCreationInfo(FReplicationBridgeSerializationContext& Context, FNetRefHandle Handle) override;
 	IRISCORE_API virtual bool CacheNetRefHandleCreationInfo(FNetRefHandle Handle) override;
 	IRISCORE_API virtual FReplicationBridgeCreateNetRefHandleResult CreateNetRefHandleFromRemote(FNetRefHandle RootObjectNetHandle, FNetRefHandle WantedNetHandle, FReplicationBridgeSerializationContext& Context) override;
 	IRISCORE_API virtual void SubObjectCreatedFromReplication(UE::Net::Private::FInternalNetRefIndex RootObjectIndex, FNetRefHandle SubObjectCreated) override;
@@ -307,6 +331,12 @@ protected:
 	IRISCORE_API virtual void OnProtocolMismatchReported(FNetRefHandle RefHandle, uint32 ConnectionId) override;
 	IRISCORE_API virtual void OnErrorWithNetRefHandleReported(UE::Net::ENetRefHandleError ErrorType, FNetRefHandle RefHandle, uint32 ConnectionId, const TArray<FNetRefHandle>& ExtraHandles) override;
 
+	/**
+	 * Remove destruction infos associated with group
+	 * Passing in an invalid group handle indicates that we should remove all destruction infos
+	 */
+	IRISCORE_API virtual void RemoveDestructionInfosForGroup(UE::Net::FNetObjectGroupHandle GroupHandle) override;
+
 private:
 
 	void OnErrorBitstreamCorrupted(FNetRefHandle RefHandle, uint32 ConnectionId);
@@ -393,6 +423,9 @@ protected:
 	 */
 	IRISCORE_API virtual TArray<uint32> FindConnectionsFromArgs(const TArray<FString>& Args) const;
 
+	/** Retrieves the dynamic filter to set for the given class. Will return an invalid handle if no dynamic filter should be set. */
+	IRISCORE_API UE::Net::FNetObjectFilterHandle GetDynamicFilter(const UClass* Class, bool bRequireForceEnabled, FName& OutFilterProfile);
+
 private:
 
 	/** Forcibly poll a single replicated object */
@@ -429,11 +462,6 @@ private:
 	/** Does all that is necessary to replicate any given UObject. Returns the object's handle if successful  */
 	UE::Net::FNetRefHandle StartReplicatingNetObject(UObject* Instance, UE::Net::EReplicationFragmentTraits Traits, UE::Net::FNetObjectFactoryId NetFactoryId);
 
-protected:
-	/** Retrieves the dynamic filter to set for the given class. Will return an invalid handle if no dynamic filter should be set. */
-	IRISCORE_API UE::Net::FNetObjectFilterHandle GetDynamicFilter(const UClass* Class, bool bRequireForceEnabled, FName& OutFilterProfile);
-
-private:
 	/** Retrieves the prioritizer to set for the given class. If bRequireForceEnabled the config needs to have bForceEnableOnAllInstances set in order for this method to return the configured prioritizer. Returns an invalid handle if no prioritizer should be set. */
 	UE::Net::FNetObjectPrioritizerHandle GetPrioritizer(const UClass* Class, bool bRequireForceEnabled);
 
@@ -465,17 +493,33 @@ private:
 	
 	FName GetConfigClassPathName(const UClass* Class);
 
+	void CreateDestructionProtocol();
 	void InitConditionalPropertyDelegates();
 	void InitNetObjectFactories();
 	void DeinitNetObjectFactories();
 
 	void OnMaxInternalNetRefIndexIncreased(UE::Net::Private::FInternalNetRefIndex NewMaxInternalIndex);
 
+private:
+
+	// Privates methods available through the FBridgeSerilization interface
+	friend UE::Net::Private::FBridgeSerialization;
+
+	/** Return the destruction info assigned to a replicated object*/
+	const UE::Net::Private::FStaticDestructionInfo* GetStaticDestructionInfo(UE::Net::FNetRefHandle Handle) const
+	{
+		return StaticObjectsPendingDestroy.Find(Handle);
+	}
+
+	/** Serialize the creation info of a replicated object */
+	bool WriteNetRefHandleCreationInfo(FReplicationBridgeSerializationContext& Context, FNetRefHandle Handle);
+
 private:
 
 	friend UE::Net::Private::FObjectPoller;
-	friend UReplicationSystem;
 	friend UE::Net::Private::FReplicationSystemImpl;
+	
+	friend UReplicationSystem;
 
 	TMap<const UClass*, FName> ConfigClassPathNameCache;
 
@@ -539,6 +583,9 @@ private:
 	// When we flush objects, we might need to defer sending creation info.
 	TMap<FNetRefHandle, TUniquePtr<const UE::Net::FNetObjectCreationHeader>> CachedCreationHeaders;
 
+	// Track the static objects that are destroyed so clients can be told to delete them after loading the level
+	TMap<FNetRefHandle, UE::Net::Private::FStaticDestructionInfo> StaticObjectsPendingDestroy;
+
 	// Objects which has object references and could be affected by garbage collection.
 	UE::Net::FNetBitArray ObjectsWithObjectReferences;
 	// Objects that needs to update their object references due to garbage collection.
@@ -550,6 +597,9 @@ private:
 	FDelegateHandle OnCustomConditionChangedHandle;
 	FDelegateHandle OnDynamicConditionChangedHandle;
 
+	// Pre-built protocol used to send destruction requests for static objects
+	const UE::Net::FReplicationProtocol* DestructionInfoProtocol = nullptr;
+	
 	UPROPERTY()
 	TArray<TObjectPtr<UNetObjectFactory>> NetObjectFactories;
 
@@ -572,4 +622,35 @@ private:
 	TArray<FString> ClassNamesToIgnoreProtocolMismatch;
 };
 
-ENUM_CLASS_FLAGS(UObjectReplicationBridge::EPrintDebugInfoTraits);
\ No newline at end of file
+//------------------------------------------------------------------------
+
+ENUM_CLASS_FLAGS(UObjectReplicationBridge::EPrintDebugInfoTraits);
+
+//------------------------------------------------------------------------
+
+namespace UE::Net::Private
+{
+
+/** Helper class to handle private functionnalities needed by Iris serialization classes */
+class FBridgeSerialization
+{
+public:
+
+	/**
+	 * Write data required to instantiate NetObject remotely to bitstream.
+	 * @param Context The serialization context parameters.
+	 * @param Handle The handle of the object to write creation data for.
+	 */
+	static bool WriteNetRefHandleCreationInfo(UObjectReplicationBridge* Bridge, FReplicationBridgeSerializationContext& Context, FNetRefHandle Handle)
+	{
+		return Bridge->WriteNetRefHandleCreationInfo(Context, Handle);
+	}
+
+	/** Return the destruction info assigned to a replicated object*/
+	static const UE::Net::Private::FStaticDestructionInfo* GetStaticDestructionInfo(UObjectReplicationBridge* Bridge, UE::Net::FNetRefHandle Handle)
+	{
+		return Bridge->GetStaticDestructionInfo(Handle);
+	}
+};
+
+} // end namespace UE::Net::Private
\ No newline at end of file
diff --git a/Engine/Source/Runtime/Net/Iris/Public/Iris/ReplicationSystem/ReplicationBridge.h b/Engine/Source/Runtime/Net/Iris/Public/Iris/ReplicationSystem/ReplicationBridge.h
--- a/Engine/Source/Runtime/Net/Iris/Public/Iris/ReplicationSystem/ReplicationBridge.h
+++ b/Engine/Source/Runtime/Net/Iris/Public/Iris/ReplicationSystem/ReplicationBridge.h
@@ -135,9 +135,6 @@ public:
 	*/
 	IRISCORE_API void StopReplicatingNetRefHandle(FNetRefHandle Handle, EEndReplicationFlags EndReplicationFlags);
 
-	/** Store destruction info for the referenced object. */
-	IRISCORE_API FNetRefHandle StoreDestructionInfo(FNetRefHandle Handle, const UE::Net::FDestructionParameters& Parameters);
-
 	/** Returns true if the handle is replicated. */
 	IRISCORE_API bool IsReplicatedHandle(FNetRefHandle Handle) const;
 
@@ -180,13 +177,6 @@ protected:
 
 	// Remote interface, invoked from Replication code during serialization
 	
-	/**
-	 * Write data required to instantiate NetObject remotely to bitstream.
-	 * @param Context The serialization context parameters.
-	 * @param Handle The handle of the object to write creation data for.
-	 */
-	IRISCORE_API virtual bool WriteNetRefHandleCreationInfo(FReplicationBridgeSerializationContext& Context, FNetRefHandle Handle);
-
 	/**
 	 * Cache info required to allow deferred writing of NetRefHandleCreationInfo
 	 * @param Handle The handle of the object to store creation data for.
@@ -283,6 +273,12 @@ protected:
 	/** Called from PostSeamlessTravelGarbageCollect. */
 	IRISCORE_API virtual void OnPostSeamlessTravelGarbageCollect();
 
+	/**
+	 * Remove destruction infos associated with group
+	 * Passing in an invalid group handle indicates that we should remove all destruction infos
+	 */
+	virtual void RemoveDestructionInfosForGroup(UE::Net::FNetObjectGroupHandle GroupHandle) {}
+
 private:
 
 	// Internal operations invoked by ReplicationSystem/ReplicationWriter
@@ -306,7 +302,6 @@ private:
 	void CallUpdateInstancesWorldLocation();
 	bool CallCacheNetRefHandleCreationInfo(FNetRefHandle Handle);
 	bool CallWriteNetRefHandleCreationInfo(FReplicationBridgeSerializationContext& Context, FNetRefHandle Handle);
-	bool CallWriteNetRefHandleDestructionInfo(FReplicationBridgeSerializationContext& Context, FNetRefHandle Handle);
 	void CallSubObjectCreatedFromReplication(UE::Net::Private::FInternalNetRefIndex RootObjectIndex, FNetRefHandle SubObjectCreated);
 	void CallPostApplyInitialState(UE::Net::Private::FInternalNetRefIndex InternalObjectIndex);
 	void CallPruneStaleObjects();
@@ -335,12 +330,6 @@ private:
 	 */
 	void NotifyStreamingLevelUnload(const UObject* Level);
 
-	/**
-	 * Remove destruction infos associated with group
-	 * Passing in an invalid group handle indicates that we should remove all destruction infos
-	 */
-	void RemoveDestructionInfosForGroup(UE::Net::FNetObjectGroupHandle GroupHandle);
-
 	void DestroyLocalNetHandle(FNetRefHandle Handle, EEndReplicationFlags Flags);
 
 	// Tear-off all handles in the PendingTearOff list that has not yet been torn-off
@@ -357,7 +346,6 @@ private:
 
 	friend UE::Net::Private::FReplicationSystemImpl;
 	friend UE::Net::Private::FReplicationSystemInternal;
-	friend UE::Net::Private::FReplicationWriter;
 	friend UE::Net::Private::FReplicationReader;
 	friend UObjectReplicationBridge;
 	friend UReplicationSystem;
@@ -379,21 +367,6 @@ private:
 
 private:
 
-	// DestructionInfos
-	struct FDestructionInfo
-	{
-		UE::Net::FNetObjectReference StaticRef;
-		UE::Net::FNetObjectGroupHandle LevelGroupHandle;
-		UE::Net::FNetObjectFactoryId NetFactoryId = UE::Net::InvalidNetObjectFactoryId;
-		UE::Net::Private::FInternalNetRefIndex InternalReplicationIndex;
-	};
-
-	const UE::Net::FReplicationProtocol* DestructionInfoProtocol;
-	
-	// Need to track the objects with destruction infos so that we can clean them up properly
-	// We use this to be able ask remote to destroy static objects
-	TMap<FNetRefHandle, FDestructionInfo> StaticObjectsPendingDestroy;
-
 	struct FPendingEndReplicationInfo
 	{
 		FPendingEndReplicationInfo(FNetRefHandle InHandle, EEndReplicationFlags InDestroyFlags, EPendingEndReplicationImmediate InImmediate) : Handle(InHandle), DestroyFlags(InDestroyFlags), Immediate(InImmediate) {}
