Iris
* Removed hack that got the default cull distance of a class by passing a CDO instance to the NetObjectFactory

#rnx

diff --git a/Engine/Source/Runtime/Net/Iris/Private/Iris/ReplicationSystem/ObjectReplicationBridgeDebugging.cpp b/Engine/Source/Runtime/Net/Iris/Private/Iris/ReplicationSystem/ObjectReplicationBridgeDebugging.cpp
--- a/Engine/Source/Runtime/Net/Iris/Private/Iris/ReplicationSystem/ObjectReplicationBridgeDebugging.cpp
+++ b/Engine/Source/Runtime/Net/Iris/Private/Iris/ReplicationSystem/ObjectReplicationBridgeDebugging.cpp
@@ -946,7 +946,6 @@ void UObjectReplicationBridge::PrintNetCullDistances(const TArray<FString>& Args
 	struct FCullDistanceInfo
 	{
 		UClass* Class = nullptr;
-		float CDOCullDistance = 0.0f;
 
 		// Total replicated root objects of this class
 		uint32 NumTotal = 0;
@@ -997,26 +996,7 @@ void UObjectReplicationBridge::PrintNetCullDistances(const TArray<FString>& Args
 			UClass* RepObjClass = RepObj->GetClass();
 
 			FCullDistanceInfo& Info = ClassCullDistanceMap.FindOrAdd(RepObjClass);
-
-			// Find the CullDistance of the CDO the first time we discover a new class 
-			if (Info.Class == nullptr)
-			{
-				UNetObjectFactory* NetFactory = GetNetFactory(NetRefHandleManager->GetReplicatedObjectData(RootObjectIndex).NetFactoryId);
-				const UNetObjectFactory::FWorldInfoContext Context
-				{
-					.Instance = RepObjClass->GetDefaultObject(), // Cheat by passing the CDO as the instance here.
-					.Handle = FNetRefHandle::GetInvalid(), // Set an invalid handle since the CDO instance is not replicated.
-					.InfoRequested = UNetObjectFactory::EWorldInfoRequested::CullDistance
-				};
-
-				Info.Class = RepObjClass;
-
-				if (TOptional<UNetObjectFactory::FWorldInfoData> OptWorldData = NetFactory->GetWorldInfo(Context))
-				{
-					Info.CDOCullDistance = OptWorldData->CullDistance;
-				}
-			}
-
+			Info.Class = RepObjClass;
 			Info.NumTotal++;
 
 			if (WorldLocations.HasCullDistanceOverride(RootObjectIndex))
