Iris
* Removed dependency to UReplicationBridge from ReplicationReader.
* Moved functions to UObjectReplicationBridge and gave the RepReader access via FBridgeSerialization

#rnx

diff --git a/Engine/Source/Runtime/Net/Iris/Private/Iris/ReplicationSystem/ObjectReplicationBridge.cpp b/Engine/Source/Runtime/Net/Iris/Private/Iris/ReplicationSystem/ObjectReplicationBridge.cpp
--- a/Engine/Source/Runtime/Net/Iris/Private/Iris/ReplicationSystem/ObjectReplicationBridge.cpp
+++ b/Engine/Source/Runtime/Net/Iris/Private/Iris/ReplicationSystem/ObjectReplicationBridge.cpp
@@ -45,6 +45,7 @@
 #include "Iris/ReplicationSystem/ReplicationWriter.h"
 #include "Iris/ReplicationSystem/RepTag.h"
 
+#include "Iris/Serialization/ObjectNetSerializer.h"
 #include "Iris/Serialization/NetSerializationContext.h"
 #include "Iris/Serialization/InternalNetSerializationContext.h"
 #include "Iris/Serialization/NetBitStreamUtil.h"
@@ -477,6 +478,48 @@ void UObjectReplicationBridge::RemoveDestructionInfosForGroup(UE::Net::FNetObjec
 	}
 }
 
+void UObjectReplicationBridge::ReadAndExecuteDestructionInfoFromRemote(FReplicationBridgeSerializationContext& Context)
+{
+	using namespace UE::Net;
+	using namespace UE::Net::Private;
+
+	check(Context.bIsDestructionInfo);
+
+	UE_NET_TRACE_SCOPE(DestructionInfo, *Context.SerializationContext.GetBitStreamReader(), Context.SerializationContext.GetTraceCollector(), ENetTraceVerbosity::Trace);
+
+	// Destroy instance here, or defer to later?
+	FNetObjectReference ReferenceToDestroy;
+	ReadFullNetObjectReference(Context.SerializationContext, ReferenceToDestroy);
+
+	const FNetObjectFactoryId NetFactoryId = IntCastChecked<FNetObjectFactoryId>(Context.SerializationContext.GetBitStreamReader()->ReadBits(FNetObjectFactoryRegistry::GetMaxBits()));
+
+	// Destroy the reference
+	// Resolve the reference in order to be able to destroy it
+	if (const UObject* Instance = ObjectReferenceCache->ResolveObjectReference(ReferenceToDestroy, Context.SerializationContext.GetInternalContext()->ResolveContext))
+	{
+		const FInternalNetRefIndex InternalReplicationIndex = NetRefHandleManager->GetInternalIndex(ReferenceToDestroy.GetRefHandle());
+		if (InternalReplicationIndex != FNetRefHandleManager::InvalidInternalIndex)
+		{
+			NetRefHandleManager->GetReplicatedObjectDataNoCheck(InternalReplicationIndex).bPendingEndReplication = 1U;
+		}
+
+		UE_LOG(LogIris, VeryVerbose, TEXT("ReadAndExecuteDestructionInfoFromRemote on %s | %s | NetFactoryId: %u"), *ReferenceToDestroy.ToString(), *ReferenceToDestroy.GetRefHandle().ToString(), NetFactoryId);
+
+		constexpr EReplicationBridgeDestroyInstanceReason DestroyReason = EReplicationBridgeDestroyInstanceReason::Destroy;
+		const EReplicationBridgeDestroyInstanceFlags DestroyFlags = IsAllowedToDestroyInstance(Instance) ? EReplicationBridgeDestroyInstanceFlags::AllowDestroyInstanceFromRemote : EReplicationBridgeDestroyInstanceFlags::None;
+
+		// if a subobject owner is to be destroyed we want to detach all subobjects before doing so to ensure we execute expected callbacks
+		// We keep tracking them internally
+		DetachSubObjectInstancesFromRemote(ReferenceToDestroy.GetRefHandle(), DestroyReason, DestroyFlags);
+
+		CallDetachInstanceFromRemote(ReferenceToDestroy.GetRefHandle(), DestroyReason, DestroyFlags, NetFactoryId);
+	}
+	else
+	{
+		//$TODO: Need to validate that it's safe if the resolve returns false
+	}
+}
+
 UObject* UObjectReplicationBridge::GetReplicatedObject(FNetRefHandle Handle) const
 {
 	return IsReplicatedHandle(Handle) ? GetObjectFromReferenceHandle(Handle) : nullptr;
@@ -867,6 +910,48 @@ UE::Net::FNetRefHandle UObjectReplicationBridge::StartReplicatingSubObject(UObje
 	return SubObjectRefHandle;
 }
 
+
+void UObjectReplicationBridge::DestroyNetObjectFromRemote(UE::Net::Private::FInternalNetRefIndex InternalObjectIndex, EReplicationBridgeDestroyInstanceReason DestroyReason)
+{
+	using namespace UE::Net::Private;
+
+	if (InternalObjectIndex == FNetRefHandleManager::InvalidInternalIndex)
+	{
+		return;
+	}
+
+	FNetRefHandleManager::FReplicatedObjectData& ObjectData = NetRefHandleManager->GetReplicatedObjectDataNoCheck(InternalObjectIndex);
+
+	ObjectData.bTearOff = (DestroyReason == EReplicationBridgeDestroyInstanceReason::TearOff);
+	ObjectData.bPendingEndReplication = 1U;
+
+	const EReplicationBridgeDestroyInstanceFlags DestroyFlags = (ObjectData.bAllowDestroyInstanceFromRemote ? EReplicationBridgeDestroyInstanceFlags::AllowDestroyInstanceFromRemote : EReplicationBridgeDestroyInstanceFlags::None);
+
+	UE_LOG_OBJECTREPLICATIONBRIDGE(Verbose, TEXT("DestroyNetObjectFromRemote for %s | DestroyReason: %s | DestroyFlags: %s "), *NetRefHandleManager->PrintObjectFromIndex(InternalObjectIndex), LexToString(DestroyReason), LexToString(DestroyFlags) );
+
+	const FNetRefHandle DestroyedObjectHandle = ObjectData.RefHandle;
+	ensureMsgf(DestroyedObjectHandle.IsValid(), TEXT("Object with InternalIndex assigned an invalid handle: %s"), *NetRefHandleManager->PrintObjectFromIndex(InternalObjectIndex));
+	if (!DestroyedObjectHandle.IsValid())
+	{
+		return;
+	}
+
+	//$IRIS TODO: Most of those operate on the InternalIndex. We should simply pass the Index instead of the Handle to them to reduce map lookups.
+
+	// if a subobject owner is to be destroyed we want to detach all subobjects before doing so to ensure we execute expected callbacks
+	// We keep tracking them internally
+	DetachSubObjectInstancesFromRemote(DestroyedObjectHandle, DestroyReason, DestroyFlags);
+		
+	// Allow derived bridges to cleanup any instance info they have stored
+	CallDetachInstanceFromRemote(DestroyedObjectHandle, DestroyReason, DestroyFlags, ObjectData.NetFactoryId);
+
+	// Detach instance protocol
+	InternalDetachInstanceFromNetRefHandle(DestroyedObjectHandle);
+	
+	// Destroy the NetRefHandle
+	InternalDestroyNetObject(DestroyedObjectHandle);
+}
+
 void UObjectReplicationBridge::SetSubObjectDormancyStatus(FNetRefHandle SubObjectRefHandle, FNetRefHandle OwnerRefHandle)
 {
 	using namespace UE::Net;
@@ -1160,6 +1245,8 @@ FReplicationBridgeCreateNetRefHandleResult UObjectReplicationBridge::CreateNetRe
 	using namespace UE::Net;
 	using namespace UE::Net::Private;
 
+	check(!Context.bIsDestructionInfo);
+
 	FNetBitStreamReader* Reader = Context.SerializationContext.GetBitStreamReader();
 
 	const FNetObjectFactoryId FactoryId = IntCastChecked<FNetObjectFactoryId>(Reader->ReadBits(FNetObjectFactoryRegistry::GetMaxBits()));
@@ -1196,7 +1283,6 @@ FReplicationBridgeCreateNetRefHandleResult UObjectReplicationBridge::CreateNetRe
 	// https://jira.it.epicgames.com/browse/UE-127369	
 
 	UObject* InstancePtr = nullptr;
-	FReplicationBridgeCreateNetRefHandleResult CreateResult;
 	
 	// Instantiate the replicated object
 	const UNetObjectFactory::FInstantiateContext InstantiateContext(WantedNetHandle, Context.SerializationContext.GetInternalContext()->ResolveContext, RootObjectOfSubObject);
@@ -1223,7 +1309,6 @@ FReplicationBridgeCreateNetRefHandleResult UObjectReplicationBridge::CreateNetRe
 	}
 
 	InstancePtr = Result.Instance;
-	CreateResult.Flags |= Result.Flags;
 
 	// Currently remote objects can only receive replicated data
 	FFragmentRegistrationContext FragmentRegistrationContext(GetReplicationStateDescriptorRegistry(), GetReplicationSystem(), EReplicationFragmentTraits::CanReceive, InstancePtr);
@@ -1316,30 +1401,48 @@ FReplicationBridgeCreateNetRefHandleResult UObjectReplicationBridge::CreateNetRe
 		.ReplicationProtocol = ReplicationProtocol,
 	};
 
-	CreateResult.NetRefHandle = InternalCreateNetObjectFromRemote(WantedNetHandle, CreateNetObjectParams);
+	const FNetRefHandle CreatedNetRefHandle = InternalCreateNetObjectFromRemote(WantedNetHandle, CreateNetObjectParams);
 
-	if (CreateResult.NetRefHandle.IsValid())
+	if (!CreatedNetRefHandle.IsValid())
 	{
-		RegisterRemoteInstance(CreateResult.NetRefHandle, InstancePtr, ReplicationProtocol, InstanceProtocol.Get(), Context.ConnectionId);
+		return FReplicationBridgeCreateNetRefHandleResult();
+	}
 
-		// Release instance protocol from the uniquePtr as it is now successfully bound to the handle
-		(void)InstanceProtocol.Release();
+	// Finalize the creation process
+	const FReplicationBridgeCreateNetRefHandleResult CreateResult
+	{
+		.NetRefHandle = CreatedNetRefHandle,
+		.Flags = Result.Flags,
+	};
 
-		// Dynamic objects or static objects that request it need to update the object cache
-		if (CreateResult.NetRefHandle.IsDynamic() || EnumHasAnyFlags(CreateResult.Flags, EReplicationBridgeCreateNetRefHandleResultFlags::BindStaticObjectInReferenceCache))
-		{
-			GetObjectReferenceCache()->BindRemoteReference(CreateResult.NetRefHandle, InstancePtr);
-		}
+	RegisterRemoteInstance(CreateResult.NetRefHandle, InstancePtr, ReplicationProtocol, InstanceProtocol.Get(), Context.ConnectionId);
 
+	// Release instance protocol from the uniquePtr as it is now successfully bound to the handle
+	(void)InstanceProtocol.Release();
+
+	// Dynamic objects or static objects that request it need to update the object cache
+	if (CreateResult.NetRefHandle.IsDynamic() || EnumHasAnyFlags(CreateResult.Flags, EReplicationBridgeCreateNetRefHandleResultFlags::BindStaticObjectInReferenceCache))
+	{
+		GetObjectReferenceCache()->BindRemoteReference(CreateResult.NetRefHandle, InstancePtr);
+	}
+
+	// Trigger the PostInstantiation callback
+	{
 		const UNetObjectFactory::FPostInstantiationContext PostContext
-		{ 
-			.Instance = InstancePtr, 
-			.Header = Header.Get(), 
-			.ConnectionId = Context.ConnectionId 
+		{
+			.Instance = InstancePtr,
+			.Header = Header.Get(),
+			.ConnectionId = Context.ConnectionId
 		};
 		Factory->PostInstantiation(PostContext);
 	}
 
+	// If it's a subobjects add the new handle to the rootobject list
+	if (RootObjectOfSubObject.IsValid())
+	{
+		NetRefHandleManager->AddSubObject(RootObjectOfSubObject, CreateResult.NetRefHandle);
+	}
+
 	return CreateResult;
 }
 
diff --git a/Engine/Source/Runtime/Net/Iris/Private/Iris/ReplicationSystem/ReplicationBridge.cpp b/Engine/Source/Runtime/Net/Iris/Private/Iris/ReplicationSystem/ReplicationBridge.cpp
--- a/Engine/Source/Runtime/Net/Iris/Private/Iris/ReplicationSystem/ReplicationBridge.cpp
+++ b/Engine/Source/Runtime/Net/Iris/Private/Iris/ReplicationSystem/ReplicationBridge.cpp
@@ -100,15 +100,6 @@ bool UReplicationBridge::CacheNetRefHandleCreationInfo(FNetRefHandle Handle)
 	return false;
 }
 
-FReplicationBridgeCreateNetRefHandleResult UReplicationBridge::CreateNetRefHandleFromRemote(FNetRefHandle RootObjectOfSubObject, FNetRefHandle WantedNetHandle, FReplicationBridgeSerializationContext& Context)
-{
-	return FReplicationBridgeCreateNetRefHandleResult();
-};
-
-void UReplicationBridge::SubObjectCreatedFromReplication(UE::Net::Private::FInternalNetRefIndex RootObjectIndex, FNetRefHandle SubObjectRefHandle)
-{
-}
-
 void UReplicationBridge::DetachInstanceFromRemote(FNetRefHandle Handle, EReplicationBridgeDestroyInstanceReason DestroyReason, EReplicationBridgeDestroyInstanceFlags DestroyFlags,  UE::Net::FNetObjectFactoryId NetFactoryId)
 {
 }
@@ -184,33 +175,6 @@ void UReplicationBridge::DetachSubObjectInstancesFromRemote(FNetRefHandle OwnerH
 	}
 }
 
-void UReplicationBridge::DestroyNetObjectFromRemote(FNetRefHandle Handle, EReplicationBridgeDestroyInstanceReason DestroyReason, EReplicationBridgeDestroyInstanceFlags DestroyFlags)
-{
-	using namespace UE::Net::Private;
-
-	if (Handle.IsValid())
-	{
-		UE_LOG_REPLICATIONBRIDGE(Verbose, TEXT("DestroyNetObjectFromRemote for %s | DestroyReason: %s | DestroyFlags: %s "), *NetRefHandleManager->PrintObjectFromNetRefHandle(Handle), LexToString(DestroyReason), LexToString(DestroyFlags) );
-
-		FInternalNetRefIndex OwnerInternalIndex = NetRefHandleManager->GetInternalIndex(Handle);
-		FNetRefHandleManager::FReplicatedObjectData& ObjectData = NetRefHandleManager->GetReplicatedObjectDataNoCheck(OwnerInternalIndex);
-		ObjectData.bTearOff = (DestroyReason == EReplicationBridgeDestroyInstanceReason::TearOff);
-		ObjectData.bPendingEndReplication = 1U;
-
-		// if a subobject owner is to be destroyed we want to detach all subobjects before doing so to ensure we execute expected callbacks
-		// We keep tracking them internally
-		DetachSubObjectInstancesFromRemote(Handle, DestroyReason, DestroyFlags);
-		
-		// Allow derived bridges to cleanup any instance info they have stored
-		CallDetachInstanceFromRemote(Handle, DestroyReason, DestroyFlags, ObjectData.NetFactoryId);
-
-		// Detach instance protocol
-		InternalDetachInstanceFromNetRefHandle(Handle);
-	
-		// Destroy the NetRefHandle
-		InternalDestroyNetObject(Handle);
-	}
-}
 
 void UReplicationBridge::CallPreSendUpdate(float DeltaSeconds)
 {
@@ -250,78 +214,11 @@ bool UReplicationBridge::CallCacheNetRefHandleCreationInfo(FNetRefHandle Handle)
 	return CacheNetRefHandleCreationInfo(Handle);
 }
 
-FReplicationBridgeCreateNetRefHandleResult UReplicationBridge::CallCreateNetRefHandleFromRemote(FNetRefHandle RootObjectOfSubObject, FNetRefHandle WantedNetHandle, FReplicationBridgeSerializationContext& Context)
-{
-	check(!Context.bIsDestructionInfo);
-
-	FReplicationBridgeCreateNetRefHandleResult CreateResult = CreateNetRefHandleFromRemote(RootObjectOfSubObject, WantedNetHandle, Context);
-
-	// Track subobjects on clients
-	if (CreateResult.NetRefHandle.IsValid() && RootObjectOfSubObject.IsValid())
-	{
-		NetRefHandleManager->AddSubObject(RootObjectOfSubObject, CreateResult.NetRefHandle);
-	}
-
-	return CreateResult;
-}
-
 bool UReplicationBridge::IsAllowedToDestroyInstance(const UObject* Instance) const
 {
 	return true;
 }
 
-void UReplicationBridge::ReadAndExecuteDestructionInfoFromRemote(FReplicationBridgeSerializationContext& Context)
-{
-	using namespace UE::Net;
-	using namespace UE::Net::Private;
-
-	check(Context.bIsDestructionInfo);
-
-	UE_NET_TRACE_SCOPE(DestructionInfo, *Context.SerializationContext.GetBitStreamReader(), Context.SerializationContext.GetTraceCollector(), ENetTraceVerbosity::Trace);
-
-	// Destroy instance here, or defer to later?
-	FNetObjectReference ReferenceToDestroy;
-	ReadFullNetObjectReference(Context.SerializationContext, ReferenceToDestroy);
-
-	const FNetObjectFactoryId NetFactoryId = IntCastChecked<FNetObjectFactoryId>(Context.SerializationContext.GetBitStreamReader()->ReadBits(FNetObjectFactoryRegistry::GetMaxBits()));
-
-	// Destroy the reference
-	// Resolve the reference in order to be able to destroy it
-	if (const UObject* Instance = ObjectReferenceCache->ResolveObjectReference(ReferenceToDestroy, Context.SerializationContext.GetInternalContext()->ResolveContext))
-	{
-		const FInternalNetRefIndex InternalReplicationIndex = NetRefHandleManager->GetInternalIndex(ReferenceToDestroy.GetRefHandle());
-		if (InternalReplicationIndex != FNetRefHandleManager::InvalidInternalIndex)
-		{
-			NetRefHandleManager->GetReplicatedObjectDataNoCheck(InternalReplicationIndex).bPendingEndReplication = 1U;
-		}
-
-		UE_LOG(LogIris, VeryVerbose, TEXT("ReadAndExecuteDestructionInfoFromRemote on %s | %s | NetFactoryId: %u"), *ReferenceToDestroy.ToString(), *ReferenceToDestroy.GetRefHandle().ToString(), NetFactoryId);
-
-		constexpr EReplicationBridgeDestroyInstanceReason DestroyReason = EReplicationBridgeDestroyInstanceReason::Destroy;
-		const EReplicationBridgeDestroyInstanceFlags DestroyFlags = IsAllowedToDestroyInstance(Instance) ? EReplicationBridgeDestroyInstanceFlags::AllowDestroyInstanceFromRemote : EReplicationBridgeDestroyInstanceFlags::None;
-
-		// if a subobject owner is to be destroyed we want to detach all subobjects before doing so to ensure we execute expected callbacks
-		// We keep tracking them internally
-		DetachSubObjectInstancesFromRemote(ReferenceToDestroy.GetRefHandle(), DestroyReason, DestroyFlags);
-
-		CallDetachInstanceFromRemote(ReferenceToDestroy.GetRefHandle(), DestroyReason, DestroyFlags, NetFactoryId);
-	}
-	else
-	{
-		//$TODO: Need to validate that it's safe if the resolve returns false
-	}
-}
-
-void UReplicationBridge::CallSubObjectCreatedFromReplication(UE::Net::Private::FInternalNetRefIndex RootObjectIndex, FNetRefHandle SubObjectCreated)
-{
-	SubObjectCreatedFromReplication(RootObjectIndex, SubObjectCreated);
-}
-
-void UReplicationBridge::CallPostApplyInitialState(UE::Net::Private::FInternalNetRefIndex InternalObjectIndex)
-{
-	PostApplyInitialState(InternalObjectIndex);
-}
-
 UReplicationBridge::~UReplicationBridge()
 {
 }
@@ -1078,32 +975,3 @@ void UReplicationBridge::PostSeamlessTravelGarbageCollect()
 void UReplicationBridge::OnPostSeamlessTravelGarbageCollect()
 {
 }
-
-namespace UE::Net::Private
-{
-
-#if UE_NET_ASYNCLOADING_DEBUG
-void FAsyncLoadingSimulator::AddStalledObject(UReplicationBridge* Bridge, FNetRefHandle ReplicatedObject)
-{
-	check(Bridge);
-	if (ReplicatedObject.IsValid())
-	{
-		Bridge->DebugObjectsForcedToStall.AddUnique(ReplicatedObject);
-	}
-}
-
-void FAsyncLoadingSimulator::RemoveStalledObject(UReplicationBridge* Bridge, FNetRefHandle ReplicatedObject)
-{
-	check(Bridge);
-	Bridge->DebugObjectsForcedToStall.Remove(ReplicatedObject);
-}
-
-void FAsyncLoadingSimulator::RemoveAllStalledObjects(UReplicationBridge* Bridge)
-{
-	check(Bridge);
-	Bridge->DebugObjectsForcedToStall.Empty();
-}
-
-#endif // UE_NET_ASYNCLOADING_DEBUG
-
-} // namespace UE::Net::Private
\ No newline at end of file
diff --git a/Engine/Source/Runtime/Net/Iris/Private/Iris/ReplicationSystem/ReplicationReader.cpp b/Engine/Source/Runtime/Net/Iris/Private/Iris/ReplicationSystem/ReplicationReader.cpp
--- a/Engine/Source/Runtime/Net/Iris/Private/Iris/ReplicationSystem/ReplicationReader.cpp
+++ b/Engine/Source/Runtime/Net/Iris/Private/Iris/ReplicationSystem/ReplicationReader.cpp
@@ -674,15 +674,15 @@ void FReplicationReader::CleanupObjectData(FReplicatedObjectInfo& ObjectInfo)
 	}
 }
 
-void FReplicationReader::EndReplication(uint32 InternalIndex, bool bTearOff, bool bDestroyInstance)
+void FReplicationReader::EndReplication(FInternalNetRefIndex InternalIndex, bool bTearOff, bool bDestroyInstance)
 {
 	if (!ensure(InternalIndex != ObjectIndexForOOBAttachment))
 	{
 		return;
 	}
+	
 	if (FReplicatedObjectInfo* ObjectInfo = ReplicatedObjects.Find(InternalIndex))
 	{
-		const FNetRefHandleManager::FReplicatedObjectData& Data = NetRefHandleManager->GetReplicatedObjectDataNoCheck(InternalIndex);
 
 		CleanupReferenceTracking(ObjectInfo);
 		Attachments.DropAllAttachments(ENetObjectAttachmentType::Normal, InternalIndex);
@@ -696,8 +696,8 @@ void FReplicationReader::EndReplication(uint32 InternalIndex, bool bTearOff, boo
 		{
 			DestroyReason = EReplicationBridgeDestroyInstanceReason::Destroy;
 		}
-		EReplicationBridgeDestroyInstanceFlags DestroyFlags = (Data.bAllowDestroyInstanceFromRemote ? EReplicationBridgeDestroyInstanceFlags::AllowDestroyInstanceFromRemote : EReplicationBridgeDestroyInstanceFlags::None);
-		ReplicationBridge->DestroyNetObjectFromRemote(Data.RefHandle, DestroyReason, DestroyFlags);
+		
+		FBridgeSerialization::DestroyNetObjectFromRemote(ReplicationBridge, InternalIndex, DestroyReason);
 
 		CleanupObjectData(*ObjectInfo);
 
@@ -796,7 +796,7 @@ FPendingBatchData* FReplicationReader::UpdateUnresolvedMustBeMappedReferences(FN
 		ENetObjectReferenceResolveResult ResolveResult = ObjectReferenceCache->ResolveObjectReference(FObjectReferenceCache::MakeNetObjectReference(Handle), ResolveContext, ResolvedObject);
 
 #if UE_NET_ASYNCLOADING_DEBUG
-		if (bForceObjectsToStall && ReplicationBridge->DebugObjectsForcedToStall.Contains(OwnerHandle))
+		if (bForceObjectsToStall && FAsyncLoadingSimulator::IsObjectStalled(ReplicationBridge, OwnerHandle))
 		{
 			Unresolved.Add(Handle);
 			continue;
@@ -996,7 +996,7 @@ uint32 FReplicationReader::ReadObjectBatch(FNetSerializationContext& Context, ui
 
 		// For destruction infos we inline the exports
 		FForceInlineExportScope ForceInlineExportScope(Context.GetInternalContext());
-		ReplicationBridge->ReadAndExecuteDestructionInfoFromRemote(BridgeContext);
+		FBridgeSerialization::ReadAndExecuteDestructionInfoFromRemote(ReplicationBridge, BridgeContext);
 
 		ReadSentinel(Context, TEXT("DestructionInfo"));
 
@@ -1324,7 +1324,7 @@ void FReplicationReader::ReadObjectInBatch(FNetSerializationContext& Context, FN
 	
 		// Get Bridge
 		FReplicationBridgeSerializationContext BridgeContext(Context, Parameters.ConnectionId);
-		const FReplicationBridgeCreateNetRefHandleResult CreateResult = ReplicationBridge->CallCreateNetRefHandleFromRemote(RootObjectOfSubObject, IncompleteHandle, BridgeContext);
+		const FReplicationBridgeCreateNetRefHandleResult CreateResult = FBridgeSerialization::CreateNetRefHandleFromRemote(ReplicationBridge, RootObjectOfSubObject, IncompleteHandle, BridgeContext);
 		FNetRefHandle NetRefHandle = CreateResult.NetRefHandle;
 		if (!NetRefHandle.IsValid())
 		{	
@@ -2223,7 +2223,7 @@ void FReplicationReader::DispatchStateData(FNetSerializationContext& Context)
 			// If the object was created this frame it's initial state is now applied
 			if (Info.bHasState && Info.bIsInitialState)
 			{
-				ReplicationBridge->CallPostApplyInitialState(Info.InternalIndex);
+				FBridgeSerialization::PostApplyInitialState(ReplicationBridge, Info.InternalIndex);
 			}
 
 			// Dispatch attachment and enqueue for later resolving
@@ -2273,7 +2273,7 @@ void FReplicationReader::DispatchStateData(FNetSerializationContext& Context)
 		{
 			if (ObjectData.SubObjectRootIndex != FNetRefHandleManager::InvalidInternalIndex)
 			{
-				ReplicationBridge->CallSubObjectCreatedFromReplication(ObjectData.SubObjectRootIndex, ObjectData.RefHandle);
+				FBridgeSerialization::SubObjectCreatedFromReplication(ReplicationBridge, ObjectData.SubObjectRootIndex, ObjectData.RefHandle);
 			}
 		}
 
@@ -2485,11 +2485,11 @@ void FReplicationReader::ResolveAndDispatchUnresolvedReferences()
 
 #if UE_NET_ASYNCLOADING_DEBUG
 		// If we want to fake async loading, test if the owner can resolve objects here
-		if (bForceObjectsToStall && !ReplicationBridge->DebugObjectsForcedToStall.IsEmpty())
+		if (bForceObjectsToStall && FAsyncLoadingSimulator::HasStalledObjects(ReplicationBridge))
 		{
 			const FInternalNetRefIndex OwnerIndex = InternalIndex;
 			const FNetRefHandle OwnerHandle = NetRefHandleManager->GetNetRefHandleFromInternalIndex(OwnerIndex);
-			if (ReplicationBridge->DebugObjectsForcedToStall.Contains(OwnerHandle))
+			if (FAsyncLoadingSimulator::IsObjectStalled(ReplicationBridge, OwnerHandle))
 			{
 				continue;
 			}
diff --git a/Engine/Source/Runtime/Net/Iris/Private/Iris/ReplicationSystem/ReplicationReader.h b/Engine/Source/Runtime/Net/Iris/Private/Iris/ReplicationSystem/ReplicationReader.h
--- a/Engine/Source/Runtime/Net/Iris/Private/Iris/ReplicationSystem/ReplicationReader.h
+++ b/Engine/Source/Runtime/Net/Iris/Private/Iris/ReplicationSystem/ReplicationReader.h
@@ -14,7 +14,7 @@
 #include "Misc/MemStack.h"
 
 // Forward declaration
-class UReplicationBridge;
+class UObjectReplicationBridge;
 
 enum class EIrisAsyncLoadingPriority : uint8;
 
@@ -188,7 +188,7 @@ private:
 	FReplicatedObjectInfo& StartReplication(uint32 InternalIndex);
 
 	// Remove tracking info for the object with InternalIndex
-	void EndReplication(uint32 InternalIndex, bool bTearOff, bool bDestroyInstance);
+	void EndReplication(FInternalNetRefIndex InternalIndex, bool bTearOff, bool bDestroyInstance);
 
 	// Free any data allocated per object
 	void CleanupObjectData(FReplicatedObjectInfo& ObjectInfo);
@@ -249,7 +249,7 @@ private:
 	FReplicationSystemInternal* ReplicationSystemInternal;
 	FNetRefHandleManager* NetRefHandleManager;
 	FReplicationStateStorage* StateStorage;
-	UReplicationBridge* ReplicationBridge;
+	UObjectReplicationBridge* ReplicationBridge;
 
 	// A cache holding unresolved handles that should be resolved each time ResolveAndDispatchUnresolvedReferences() is called.
 	TMap<FNetRefHandle, uint32> HotUnresolvedHandleCache;
diff --git a/Engine/Source/Runtime/Net/Iris/Public/Iris/ReplicationSystem/ObjectReplicationBridge.h b/Engine/Source/Runtime/Net/Iris/Public/Iris/ReplicationSystem/ObjectReplicationBridge.h
--- a/Engine/Source/Runtime/Net/Iris/Public/Iris/ReplicationSystem/ObjectReplicationBridge.h
+++ b/Engine/Source/Runtime/Net/Iris/Public/Iris/ReplicationSystem/ObjectReplicationBridge.h
@@ -213,6 +213,9 @@ public:
 	/** Get the Object from a handle, even if the object was pre-registered and hasn't replicated yet. */
 	IRISCORE_API UObject* GetPreRegisteredObject(FNetRefHandle Handle) const;
 
+	/** Tell the remote connection that we detected a reading error with a specific replicated object */
+	virtual void SendErrorWithNetRefHandle(UE::Net::ENetRefHandleError ErrorType, FNetRefHandle RefHandle, uint32 ConnectionId, TConstArrayView<const FNetRefHandle> ExtraNetRefHandle = {}) {}
+
 	/**
 	 * Creates a NetRefHandle explicitly for object Instance and flags it as pre-registered in the ObjectReferenceCache.
 	 * If supported by an object factory, the pre-registered flag is meant to be serialized in a root object creation header
@@ -323,9 +326,6 @@ protected:
 	IRISCORE_API virtual void UpdateInstancesWorldLocation() override;
 	IRISCORE_API virtual void PruneStaleObjects() override;	
 	IRISCORE_API virtual bool CacheNetRefHandleCreationInfo(FNetRefHandle Handle) override;
-	IRISCORE_API virtual FReplicationBridgeCreateNetRefHandleResult CreateNetRefHandleFromRemote(FNetRefHandle RootObjectNetHandle, FNetRefHandle WantedNetHandle, FReplicationBridgeSerializationContext& Context) override;
-	IRISCORE_API virtual void SubObjectCreatedFromReplication(UE::Net::Private::FInternalNetRefIndex RootObjectIndex, FNetRefHandle SubObjectCreated) override;
-	IRISCORE_API virtual void PostApplyInitialState(UE::Net::Private::FInternalNetRefIndex InternalObjectIndex) final;
 	IRISCORE_API virtual void DetachInstanceFromRemote(FNetRefHandle Handle, EReplicationBridgeDestroyInstanceReason DestroyReason, EReplicationBridgeDestroyInstanceFlags DestroyFlags,  UE::Net::FNetObjectFactoryId NetFactoryId) override;
 	IRISCORE_API virtual void DetachInstance(FNetRefHandle Handle) override;
 	IRISCORE_API virtual void OnProtocolMismatchReported(FNetRefHandle RefHandle, uint32 ConnectionId) override;
@@ -502,18 +502,43 @@ private:
 
 private:
 
-	// Privates methods available through the FBridgeSerilization interface
+	//------------------------------------------------------------------------
+	// These private methods are available through the FBridgeSerialization interface
+	//------------------------------------------------------------------------
+
 	friend UE::Net::Private::FBridgeSerialization;
 
-	/** Return the destruction info assigned to a replicated object*/
+	/** Return the destruction info assigned to a replicated object */
 	const UE::Net::Private::FStaticDestructionInfo* GetStaticDestructionInfo(UE::Net::FNetRefHandle Handle) const
 	{
 		return StaticObjectsPendingDestroy.Find(Handle);
 	}
 
-	/** Serialize the creation info of a replicated object */
+	/**
+	 * Write data required to instantiate NetObject remotely to bitstream.
+	 * @param Context The serialization context parameters.
+	 * @param Handle The handle of the object to write creation data for.
+	 */
 	bool WriteNetRefHandleCreationInfo(FReplicationBridgeSerializationContext& Context, FNetRefHandle Handle);
 
+	/**
+	 * Client only function that executes the EndReplication request by the authority
+	 * @param Handle The netobject to destroy
+	 * @DestroyReason How the object should be destroyed
+	 */
+	void DestroyNetObjectFromRemote(UE::Net::Private::FInternalNetRefIndex InternalObjectIndex, EReplicationBridgeDestroyInstanceReason DestroyReason);
+
+	void ReadAndExecuteDestructionInfoFromRemote(FReplicationBridgeSerializationContext& Context);
+
+	/** Read data required to instantiate NetObject from bitstream. */
+	FReplicationBridgeCreateNetRefHandleResult CreateNetRefHandleFromRemote(FNetRefHandle RootObjectNetHandle, FNetRefHandle WantedNetHandle, FReplicationBridgeSerializationContext& Context);
+
+	/** Invoked after we have applied the initial state for an object.*/
+	void PostApplyInitialState(UE::Net::Private::FInternalNetRefIndex InternalObjectIndex);
+
+	/** Invoked right before we apply the state for a new received subobject but after we have applied state on the root object. */
+	void SubObjectCreatedFromReplication(UE::Net::Private::FInternalNetRefIndex RootObjectIndex, FNetRefHandle SubObjectCreated);
+
 private:
 
 	friend UE::Net::Private::FObjectPoller;
@@ -613,6 +638,16 @@ protected:
 	bool bSuppressCreateInstanceFailedEnsure = false;
 
 private:
+
+#if UE_NET_ASYNCLOADING_DEBUG
+	friend UE::Net::Private::FAsyncLoadingSimulator;
+	/** 
+	 * Debug array that forces objects to be stalled, as if async loading external references.
+	 * Only enabled when net.Iris.AsyncLoading.ForceObjectsToStall is enabled.
+	 */
+	TArray<FNetRefHandle> DebugObjectsForcedToStall;
+#endif
+
 	TMap<FObjectKey, bool> ArchetypesAlreadyPrinted;
 
 	/** Tells if we already logged a warning about a given object or a class type that had a location that is out of bounds. */
@@ -631,26 +666,89 @@ ENUM_CLASS_FLAGS(UObjectReplicationBridge::EPrintDebugInfoTraits);
 namespace UE::Net::Private
 {
 
+#if UE_NET_ASYNCLOADING_DEBUG
+struct FAsyncLoadingSimulator
+{
+	static void AddStalledObject(UObjectReplicationBridge* Bridge, FNetRefHandle ReplicatedObject)
+	{
+		check(Bridge);
+		if (ReplicatedObject.IsValid())
+		{
+			Bridge->DebugObjectsForcedToStall.AddUnique(ReplicatedObject);
+		}
+	}
+
+	static void RemoveStalledObject(UObjectReplicationBridge* Bridge, FNetRefHandle ReplicatedObject)
+	{
+		check(Bridge);
+		Bridge->DebugObjectsForcedToStall.Remove(ReplicatedObject);
+	}
+
+	static void RemoveAllStalledObjects(UObjectReplicationBridge* Bridge)
+	{
+		check(Bridge);
+		Bridge->DebugObjectsForcedToStall.Empty();
+	}
+
+	static bool HasStalledObjects(UObjectReplicationBridge* Bridge)
+	{
+		check(Bridge);
+		return !Bridge->DebugObjectsForcedToStall.IsEmpty();
+	}
+
+	static bool IsObjectStalled(UObjectReplicationBridge* Bridge, FNetRefHandle Handle)
+	{
+		check(Bridge);
+		return Bridge->DebugObjectsForcedToStall.Contains(Handle);
+	}
+
+	static FReplicationBridgeCreateNetRefHandleResult CreateNetRefHandleFromRemote(UObjectReplicationBridge* Bridge, FNetRefHandle RootObjectNetHandle, FNetRefHandle WantedNetHandle, FReplicationBridgeSerializationContext& Context)
+	{
+		check(Bridge);
+		return Bridge->CreateNetRefHandleFromRemote(RootObjectNetHandle, WantedNetHandle, Context);
+	}
+};
+#endif // UE_NET_ASYNCLOADING_DEBUG
+
 /** Helper class to handle private functionnalities needed by Iris serialization classes */
 class FBridgeSerialization
 {
 public:
 
-	/**
-	 * Write data required to instantiate NetObject remotely to bitstream.
-	 * @param Context The serialization context parameters.
-	 * @param Handle The handle of the object to write creation data for.
-	 */
 	static bool WriteNetRefHandleCreationInfo(UObjectReplicationBridge* Bridge, FReplicationBridgeSerializationContext& Context, FNetRefHandle Handle)
 	{
 		return Bridge->WriteNetRefHandleCreationInfo(Context, Handle);
 	}
 
-	/** Return the destruction info assigned to a replicated object*/
 	static const UE::Net::Private::FStaticDestructionInfo* GetStaticDestructionInfo(UObjectReplicationBridge* Bridge, UE::Net::FNetRefHandle Handle)
 	{
 		return Bridge->GetStaticDestructionInfo(Handle);
 	}
+
+	static void DestroyNetObjectFromRemote(UObjectReplicationBridge* Bridge, FInternalNetRefIndex InternalObjectIndex, EReplicationBridgeDestroyInstanceReason DestroyReason)
+	{
+		Bridge->DestroyNetObjectFromRemote(InternalObjectIndex, DestroyReason);
+	}
+
+	static void ReadAndExecuteDestructionInfoFromRemote(UObjectReplicationBridge* Bridge, FReplicationBridgeSerializationContext& Context)
+	{
+		Bridge->ReadAndExecuteDestructionInfoFromRemote(Context);
+	}
+
+	static FReplicationBridgeCreateNetRefHandleResult CreateNetRefHandleFromRemote(UObjectReplicationBridge* Bridge, FNetRefHandle RootObjectNetHandle, FNetRefHandle WantedNetHandle, FReplicationBridgeSerializationContext& Context)
+	{
+		return Bridge->CreateNetRefHandleFromRemote(RootObjectNetHandle, WantedNetHandle, Context);
+	}
+
+	static void PostApplyInitialState(UObjectReplicationBridge* Bridge,  FInternalNetRefIndex InternalObjectIndex)
+	{
+		return Bridge->PostApplyInitialState(InternalObjectIndex);
+	}
+
+	static void SubObjectCreatedFromReplication(UObjectReplicationBridge* Bridge, UE::Net::Private::FInternalNetRefIndex RootObjectIndex, FNetRefHandle SubObjectCreated)
+	{
+		return Bridge->SubObjectCreatedFromReplication(RootObjectIndex, SubObjectCreated);
+	}
 };
 
 } // end namespace UE::Net::Private
\ No newline at end of file
diff --git a/Engine/Source/Runtime/Net/Iris/Public/Iris/ReplicationSystem/ReplicationBridge.h b/Engine/Source/Runtime/Net/Iris/Public/Iris/ReplicationSystem/ReplicationBridge.h
--- a/Engine/Source/Runtime/Net/Iris/Public/Iris/ReplicationSystem/ReplicationBridge.h
+++ b/Engine/Source/Runtime/Net/Iris/Public/Iris/ReplicationSystem/ReplicationBridge.h
@@ -184,15 +184,6 @@ protected:
 	*/
 	IRISCORE_API virtual bool CacheNetRefHandleCreationInfo(FNetRefHandle Handle);
 
-	/** Read data required to instantiate NetObject from bitstream. */
-	IRISCORE_API virtual FReplicationBridgeCreateNetRefHandleResult CreateNetRefHandleFromRemote(FNetRefHandle RootObjectOfSubObject, FNetRefHandle WantedNetHandle, FReplicationBridgeSerializationContext& Context);
-
-	/** Invoked right before we apply the state for a new received subobject but after we have applied state on the root object. */
- 	IRISCORE_API virtual void SubObjectCreatedFromReplication(UE::Net::Private::FInternalNetRefIndex RootObjectIndex,FNetRefHandle SubObjectCreated);
-
-	/** Invoked after we have applied the initial state for an object.*/
-	virtual void PostApplyInitialState(UE::Net::Private::FInternalNetRefIndex InternalObjectIndex) {}
-
 	/**
 	 * Called when the instance is detached from the protocol on request by the remote. 
 	 * @param Handle The handle of the object to destroy or tear off.
@@ -264,9 +255,6 @@ protected:
 	/** Called when a remote connection has a critical error caused by a specific NetRefHandle */
 	virtual void OnErrorWithNetRefHandleReported(UE::Net::ENetRefHandleError ErrorType, FNetRefHandle RefHandle, uint32 ConnectionId, const TArray<FNetRefHandle>& ExtraNetRefHandle) {}
 
-	/** Tell the remote connection that we detected a reading error with a specific replicated object */
-	virtual void SendErrorWithNetRefHandle(UE::Net::ENetRefHandleError ErrorType, FNetRefHandle RefHandle, uint32 ConnectionId, TConstArrayView<const FNetRefHandle> ExtraNetRefHandle = {}) {}
-
 	/** Called from PreSeamlessTravelGarbageCollect. */
 	IRISCORE_API virtual void OnPreSeamlessTravelGarbageCollect();
 
@@ -282,10 +270,9 @@ protected:
 private:
 
 	// Internal operations invoked by ReplicationSystem/ReplicationWriter
-	void ReadAndExecuteDestructionInfoFromRemote(FReplicationBridgeSerializationContext& Context);
+	
 	void DetachSubObjectInstancesFromRemote(FNetRefHandle Handle, EReplicationBridgeDestroyInstanceReason DestroyReason, EReplicationBridgeDestroyInstanceFlags DestroyFlags);
-	void DestroyNetObjectFromRemote(FNetRefHandle Handle, EReplicationBridgeDestroyInstanceReason DestroyReason, EReplicationBridgeDestroyInstanceFlags DestroyFlags);
-
+	
 	enum class EPendingEndReplicationImmediate : uint8
 	{
 		Yes,
@@ -296,14 +283,11 @@ private:
 	// it will be kept around until the handle is no longer ref-counted by any connection. It will however be removed from the set of scopeable objects after the first update so new connections will not add it to their scope.
 	void AddPendingEndReplication(FNetRefHandle Handle, EEndReplicationFlags DestroyFlags, EPendingEndReplicationImmediate Immediate = EPendingEndReplicationImmediate::No);
 
-	FReplicationBridgeCreateNetRefHandleResult CallCreateNetRefHandleFromRemote(FNetRefHandle RootObjectOfSubObject, FNetRefHandle WantedNetHandle, FReplicationBridgeSerializationContext& Context);
 	void CallPreSendUpdate(float DeltaSeconds);	
 	void CallPreSendUpdateSingleHandle(FNetRefHandle Handle);
 	void CallUpdateInstancesWorldLocation();
 	bool CallCacheNetRefHandleCreationInfo(FNetRefHandle Handle);
 	bool CallWriteNetRefHandleCreationInfo(FReplicationBridgeSerializationContext& Context, FNetRefHandle Handle);
-	void CallSubObjectCreatedFromReplication(UE::Net::Private::FInternalNetRefIndex RootObjectIndex, FNetRefHandle SubObjectCreated);
-	void CallPostApplyInitialState(UE::Net::Private::FInternalNetRefIndex InternalObjectIndex);
 	void CallPruneStaleObjects();
 	void CallGetInitialDependencies(FNetRefHandle Handle, FNetDependencyInfoArray& OutDependencies) const;
 	void CallDetachInstance(FNetRefHandle Handle);
@@ -346,7 +330,7 @@ private:
 
 	friend UE::Net::Private::FReplicationSystemImpl;
 	friend UE::Net::Private::FReplicationSystemInternal;
-	friend UE::Net::Private::FReplicationReader;
+
 	friend UObjectReplicationBridge;
 	friend UReplicationSystem;
 
@@ -377,29 +361,8 @@ private:
 	};
 	TArray<FPendingEndReplicationInfo> HandlesPendingEndReplication;
 
-#if UE_NET_ASYNCLOADING_DEBUG
-	friend UE::Net::Private::FAsyncLoadingSimulator;
-	/** 
-	 * Debug array that forces objects to be stalled, as if async loading external references.
-	 * Only enabled when net.Iris.AsyncLoading.ForceObjectsToStall is enabled.
-	 */
-	TArray<FNetRefHandle> DebugObjectsForcedToStall;
-#endif
-
 };
 
-namespace UE::Net::Private
-{
-#if UE_NET_ASYNCLOADING_DEBUG
-	struct FAsyncLoadingSimulator
-	{
-		IRISCORE_API static void AddStalledObject(UReplicationBridge* Bridge, FNetRefHandle ReplicatedObject);
-		IRISCORE_API static void RemoveStalledObject(UReplicationBridge* Bridge, FNetRefHandle ReplicatedObject);
-		IRISCORE_API static void RemoveAllStalledObjects(UReplicationBridge* Bridge);
-	};
-#endif
-
-}
 
 inline FReplicationBridgeSerializationContext::FReplicationBridgeSerializationContext(UE::Net::FNetSerializationContext& InSerialiazationContext, uint32 InConnectionId, bool bInIsDestructionInfo)
 : SerializationContext(InSerialiazationContext)
