Iris - Do not allow circular object dependencies.

diff --git a/Engine/Plugins/Runtime/ReplicationSystemTestPlugin/Source/Private/Tests/ReplicationSystem/TestDependentObjects.cpp b/Engine/Plugins/Runtime/ReplicationSystemTestPlugin/Source/Private/Tests/ReplicationSystem/TestDependentObjects.cpp
--- a/Engine/Plugins/Runtime/ReplicationSystemTestPlugin/Source/Private/Tests/ReplicationSystem/TestDependentObjects.cpp
+++ b/Engine/Plugins/Runtime/ReplicationSystemTestPlugin/Source/Private/Tests/ReplicationSystem/TestDependentObjects.cpp
@@ -1,11 +1,12 @@
 // Copyright Epic Games, Inc. All Rights Reserved.
 
 #include "ReplicationSystemServerClientTestFixture.h"
+#include "Iris/Core/IrisLog.h"
 #include "Iris/ReplicationSystem/ReplicationSystem.h"
 #include "Iris/ReplicationSystem/Filtering/NetObjectFilterDefinitions.h"
+#include "Logging/LogScopedVerbosityOverride.h"
 #include "Tests/ReplicationSystem/Filtering/MockNetObjectFilter.h"
 
-
 namespace UE::Net::Private
 {
 	
@@ -945,4 +946,31 @@ UE_NET_TEST_FIXTURE(FReplicationSystemServerClientTestFixture, TestDependentObje
 	UE_NET_ASSERT_LT(ClientDependentObject->LastRepOrderCounter, ClientObject->LastRepOrderCounter);
 }
 
+UE_NET_TEST_FIXTURE(FReplicationSystemServerClientTestFixture, TestCircularDependencyIsDenied)
+{
+	UReplicatedTestObjectBridge* Bridge = Server->GetReplicationBridge();
+
+	// Spawn at least two objects on the server
+	constexpr uint32 ObjectCount = 4;
+	UTestReplicatedIrisObject* ServerObjects[ObjectCount];
+	for (UTestReplicatedIrisObject*& ServerObject : ServerObjects)
+	{
+		ServerObject = Server->CreateObject();
+	}
+
+	// Add all dependencies but the one creating the circular dependency
+	for (uint32 ObjectIt = 0, ObjectEndIt = ObjectCount - 1; ObjectIt < ObjectEndIt; ++ObjectIt)
+	{
+		Bridge->AddDependentObject(ServerObjects[ObjectIt]->NetRefHandle, ServerObjects[ObjectIt + 1]->NetRefHandle);
+	}
+
+	// Create circular dependency
+	{
+		LOG_SCOPE_VERBOSITY_OVERRIDE(LogIris, ELogVerbosity::Fatal);
+		FTestEnsureScope SuppressEnsureScope;
+		Bridge->AddDependentObject(ServerObjects[ObjectCount - 1]->NetRefHandle, ServerObjects[0]->NetRefHandle);
+		UE_NET_ASSERT_EQ(SuppressEnsureScope.GetCount(), 1);
+	}
+}
+
 }
diff --git a/Engine/Source/Runtime/Net/Iris/Private/Iris/ReplicationSystem/NetRefHandleManager.cpp b/Engine/Source/Runtime/Net/Iris/Private/Iris/ReplicationSystem/NetRefHandleManager.cpp
--- a/Engine/Source/Runtime/Net/Iris/Private/Iris/ReplicationSystem/NetRefHandleManager.cpp
+++ b/Engine/Source/Runtime/Net/Iris/Private/Iris/ReplicationSystem/NetRefHandleManager.cpp
@@ -861,6 +861,14 @@ bool FNetRefHandleManager::AddDependentObject(FNetRefHandle ParentRefHandle, FNe
 	// If child was already set as dependent on the parent there is a logic error somewhere.
 	checkf(DependentParentObjectArray.Find(ParentInternalIndex) == INDEX_NONE, TEXT("FNetRefHandleManager::AddDependentObject: Child: %s already dependent of Parent: %s but not the inverse."), *PrintObjectFromNetRefHandle(DependentObjectRefHandle), *PrintObjectFromNetRefHandle(ParentRefHandle));
 
+	// Make sure adding this dependency wouldn't create a circular dependency chain.
+	if (AddingDependencyWouldCreateCircularDependency(ParentInternalIndex, DependentObjectInternalIndex))
+	{
+		UE_LOG(LogIris, Warning, TEXT("FNetRefHandleManager::AddDependentObject: Child: %s has an implicit or direct dependency on Parent: %s. Denying circular dependency."), *PrintObjectFromNetRefHandle(ParentRefHandle), *PrintObjectFromNetRefHandle(DependentObjectRefHandle));
+		ensureMsgf(false, TEXT("FNetRefHandleManager::AddDependentObject: Child: %s has an implicit or direct dependency on Parent: %s. Denying circular dependency."), *PrintObjectFromNetRefHandle(ParentRefHandle), *PrintObjectFromNetRefHandle(DependentObjectRefHandle));
+		return false;
+	}
+
 	// Add dependent to parent's dependent object list
 	FDependentObjectInfo DependentObjectInfo;
 	DependentObjectInfo.NetRefIndex = DependentObjectInternalIndex;
@@ -1066,6 +1074,36 @@ void FNetRefHandleManager::RemoveDependentObject(FNetRefHandle ParentHandle, FNe
 	InternalRemoveDependentObject(ParentInternalIndex, DependentInternalIndex);
 }
 
+bool FNetRefHandleManager::AddingDependencyWouldCreateCircularDependency(FInternalNetRefIndex ParentInternalIndex, FInternalNetRefIndex DependentInternalIndex) const
+{
+	TArray<FInternalNetRefIndex, TInlineAllocator<64>> DependentObjects;
+	DependentObjects.Push(DependentInternalIndex);
+
+	do 
+	{
+		if (const FInternalNetRefIndex ObjectIndex = DependentObjects.Pop(EAllowShrinking::No))
+		{
+			if (ObjectsWithDependentObjectsInternalIndices.GetBit(ObjectIndex))
+			{
+				TConstArrayView<FDependentObjectInfo> DependentObjectInfos = GetDependentObjectInfos(ObjectIndex);
+				DependentObjects.Reserve(DependentObjects.Num() + DependentObjectInfos.Num());
+				for (const FDependentObjectInfo& DependentObjectInfo : DependentObjectInfos)
+				{
+					if (DependentObjectInfo.NetRefIndex == ParentInternalIndex)
+					{
+						return true;
+					}
+
+					DependentObjects.Add(DependentObjectInfo.NetRefIndex);
+				}
+			}
+		}
+	}
+	while (!DependentObjects.IsEmpty());
+
+	return false;
+}
+
 void FNetRefHandleManager::SetShouldPropagateChangedStates(FInternalNetRefIndex ObjectInternalIndex, bool bShouldPropagateChangedStates)
 {
 	if (ObjectInternalIndex != InvalidInternalIndex)
diff --git a/Engine/Source/Runtime/Net/Iris/Private/Iris/ReplicationSystem/NetRefHandleManager.h b/Engine/Source/Runtime/Net/Iris/Private/Iris/ReplicationSystem/NetRefHandleManager.h
--- a/Engine/Source/Runtime/Net/Iris/Private/Iris/ReplicationSystem/NetRefHandleManager.h
+++ b/Engine/Source/Runtime/Net/Iris/Private/Iris/ReplicationSystem/NetRefHandleManager.h
@@ -416,6 +416,7 @@ private:
 
 	void InternalRemoveDependentObject(FInternalNetRefIndex ParentInternalIndex, FInternalNetRefIndex DependentInternalIndex, ERemoveDependentObjectFlags Flags = ERemoveDependentObjectFlags::All);
 	void InternalRemoveAllDependencies(FInternalNetRefIndex DependentInternalIndex);
+	bool AddingDependencyWouldCreateCircularDependency(FInternalNetRefIndex ParentInternalIndex, FInternalNetRefIndex DependentInternalIndex) const;
 
 	/**
 	 * Grow the NetObjectLists to fit more internal indexes.
