Fixed logging.

#rnx

diff --git a/Engine/Source/Runtime/Net/Iris/Private/Iris/ReplicationSystem/NetRefHandleManager.cpp b/Engine/Source/Runtime/Net/Iris/Private/Iris/ReplicationSystem/NetRefHandleManager.cpp
--- a/Engine/Source/Runtime/Net/Iris/Private/Iris/ReplicationSystem/NetRefHandleManager.cpp
+++ b/Engine/Source/Runtime/Net/Iris/Private/Iris/ReplicationSystem/NetRefHandleManager.cpp
@@ -864,8 +864,8 @@ bool FNetRefHandleManager::AddDependentObject(FNetRefHandle ParentRefHandle, FNe
 	// Make sure adding this dependency wouldn't create a circular dependency chain.
 	if (AddingDependencyWouldCreateCircularDependency(ParentInternalIndex, DependentObjectInternalIndex))
 	{
-		UE_LOG(LogIris, Warning, TEXT("FNetRefHandleManager::AddDependentObject: Child: %s has an implicit or direct dependency on Parent: %s. Denying circular dependency."), *PrintObjectFromNetRefHandle(ParentRefHandle), *PrintObjectFromNetRefHandle(DependentObjectRefHandle));
-		ensureMsgf(false, TEXT("FNetRefHandleManager::AddDependentObject: Child: %s has an implicit or direct dependency on Parent: %s. Denying circular dependency."), *PrintObjectFromNetRefHandle(ParentRefHandle), *PrintObjectFromNetRefHandle(DependentObjectRefHandle));
+		UE_LOG(LogIris, Warning, TEXT("FNetRefHandleManager::AddDependentObject: Child: %s has an implicit or direct dependency on Parent: %s. Denying circular dependency."), *PrintObjectFromNetRefHandle(DependentObjectRefHandle), *PrintObjectFromNetRefHandle(ParentRefHandle));
+		ensureMsgf(false, TEXT("FNetRefHandleManager::AddDependentObject: Child: %s has an implicit or direct dependency on Parent: %s. Denying circular dependency."), *PrintObjectFromNetRefHandle(DependentObjectRefHandle), *PrintObjectFromNetRefHandle(ParentRefHandle));
 		return false;
 	}
 
