Iris
* Moved more functions from UBridge to UObjectBridge
* Refactored a few functions to receive an InternalIndex instead of an Handle since the call already knew the index
* Deleted CallGetInitialDependencies since it was private and not used by anyone

#rnx

diff --git a/Engine/Source/Runtime/Net/Iris/Private/Iris/ReplicationSystem/ObjectReplicationBridge.cpp b/Engine/Source/Runtime/Net/Iris/Private/Iris/ReplicationSystem/ObjectReplicationBridge.cpp
--- a/Engine/Source/Runtime/Net/Iris/Private/Iris/ReplicationSystem/ObjectReplicationBridge.cpp
+++ b/Engine/Source/Runtime/Net/Iris/Private/Iris/ReplicationSystem/ObjectReplicationBridge.cpp
@@ -510,9 +510,9 @@ void UObjectReplicationBridge::ReadAndExecuteDestructionInfoFromRemote(FReplicat
 
 		// if a subobject owner is to be destroyed we want to detach all subobjects before doing so to ensure we execute expected callbacks
 		// We keep tracking them internally
-		DetachSubObjectInstancesFromRemote(ReferenceToDestroy.GetRefHandle(), DestroyReason, DestroyFlags);
+		DetachSubObjectInstancesFromRemote(InternalReplicationIndex, DestroyReason, DestroyFlags);
 
-		CallDetachInstanceFromRemote(ReferenceToDestroy.GetRefHandle(), DestroyReason, DestroyFlags, NetFactoryId);
+		DetachInstanceFromRemote(ReferenceToDestroy.GetRefHandle(), DestroyReason, DestroyFlags, NetFactoryId);
 	}
 	else
 	{
@@ -936,17 +936,15 @@ void UObjectReplicationBridge::DestroyNetObjectFromRemote(UE::Net::Private::FInt
 		return;
 	}
 
-	//$IRIS TODO: Most of those operate on the InternalIndex. We should simply pass the Index instead of the Handle to them to reduce map lookups.
-
 	// if a subobject owner is to be destroyed we want to detach all subobjects before doing so to ensure we execute expected callbacks
 	// We keep tracking them internally
-	DetachSubObjectInstancesFromRemote(DestroyedObjectHandle, DestroyReason, DestroyFlags);
+	DetachSubObjectInstancesFromRemote(InternalObjectIndex, DestroyReason, DestroyFlags);
 		
 	// Allow derived bridges to cleanup any instance info they have stored
-	CallDetachInstanceFromRemote(DestroyedObjectHandle, DestroyReason, DestroyFlags, ObjectData.NetFactoryId);
+	DetachInstanceFromRemote(DestroyedObjectHandle, DestroyReason, DestroyFlags, ObjectData.NetFactoryId);
 
 	// Detach instance protocol
-	InternalDetachInstanceFromNetRefHandle(DestroyedObjectHandle);
+	InternalDetachInstanceFromNetRefHandle(InternalObjectIndex);
 	
 	// Destroy the NetRefHandle
 	InternalDestroyNetObject(DestroyedObjectHandle);
@@ -1204,8 +1202,43 @@ void UObjectReplicationBridge::DetachInstanceFromRemote(FNetRefHandle Handle, ER
 	}
 	
 	UnregisterInstance(Handle);	
+}
+
+void UObjectReplicationBridge::DetachSubObjectInstancesFromRemote(UE::Net::Private::FInternalNetRefIndex InternalObjectIndex, EReplicationBridgeDestroyInstanceReason DestroyReason, EReplicationBridgeDestroyInstanceFlags DestroyFlags)
+{
+	using namespace UE::Net::Private;
 
-	// $IRIS TODO: Cleanup any pending creation data if we have not yet instantiated the instance.
+	if (InternalObjectIndex == FNetRefHandleManager::InvalidInternalIndex)
+	{
+		return;
+	}
+
+	for (const FInternalNetRefIndex SubObjectInternalIndex : NetRefHandleManager->GetChildSubObjects(InternalObjectIndex))
+	{
+		FNetRefHandleManager::FReplicatedObjectData& SubObjectData = NetRefHandleManager->GetReplicatedObjectDataNoCheck(SubObjectInternalIndex);
+		const FNetRefHandle SubObjectHandle = SubObjectData.RefHandle;
+		SubObjectData.bTearOff = (DestroyReason == EReplicationBridgeDestroyInstanceReason::TearOff);
+		SubObjectData.bPendingEndReplication = 1U;
+
+		EReplicationBridgeDestroyInstanceFlags SubObjectDestroyFlags = DestroyFlags;
+		// The subobject is allowed to be destroyed if both the owner and the subobject allows it.
+		if (!SubObjectData.bAllowDestroyInstanceFromRemote)
+		{
+			EnumRemoveFlags(SubObjectDestroyFlags, EReplicationBridgeDestroyInstanceFlags::AllowDestroyInstanceFromRemote);
+		}
+
+		if (DestroyReason == EReplicationBridgeDestroyInstanceReason::DoNotDestroy && SubObjectData.bAllowDestroyInstanceFromRemote)
+		{
+			// When ending replication of static objects without destroying them, we should always destroy instances spawned from replication
+			// as they will be recreated when the static object is scoped again.
+			EnumAddFlags(SubObjectDestroyFlags, EReplicationBridgeDestroyInstanceFlags::AllowDestroyInstanceFromRemote);
+			DetachInstanceFromRemote(SubObjectHandle, EReplicationBridgeDestroyInstanceReason::Destroy, SubObjectDestroyFlags, SubObjectData.NetFactoryId);
+		}
+		else
+		{
+			DetachInstanceFromRemote(SubObjectHandle, DestroyReason, SubObjectDestroyFlags, SubObjectData.NetFactoryId);
+		}
+	}
 }
 
 void UObjectReplicationBridge::DetachInstance(FNetRefHandle RefHandle)
diff --git a/Engine/Source/Runtime/Net/Iris/Private/Iris/ReplicationSystem/ReplicationBridge.cpp b/Engine/Source/Runtime/Net/Iris/Private/Iris/ReplicationSystem/ReplicationBridge.cpp
--- a/Engine/Source/Runtime/Net/Iris/Private/Iris/ReplicationSystem/ReplicationBridge.cpp
+++ b/Engine/Source/Runtime/Net/Iris/Private/Iris/ReplicationSystem/ReplicationBridge.cpp
@@ -56,13 +56,6 @@ static FAutoConsoleVariableRef CVarAllowDestroyToCancelFlushAndTearOff(
 	TEXT("When true issuing a EndReplication on an object that is already Tear-off or pending endreplication will cancel destroy/flush and destroy the replicated objects.")
 );
 
-static bool bAlwaysDestroyDynamicSubObjectInstancesOnDetachFromStaticRoot = true;
-static FAutoConsoleVariableRef CVarAlwaysDestroyDynamicSubObjectInstancesOnDetachFromStaticRoot(
-	TEXT("net.Iris.AlwaysDestroyDynamicSubObjectInstancesOnDetachFromStaticRoot"),
-	bAlwaysDestroyDynamicSubObjectInstancesOnDetachFromStaticRoot,
-	TEXT("When true, We will always destroy instance for dynamic subobjects during EndReplication of a static rootobject.")
-);
-
 /**
  * ReplicationBridge Implementation
  */
@@ -100,10 +93,6 @@ bool UReplicationBridge::CacheNetRefHandleCreationInfo(FNetRefHandle Handle)
 	return false;
 }
 
-void UReplicationBridge::DetachInstanceFromRemote(FNetRefHandle Handle, EReplicationBridgeDestroyInstanceReason DestroyReason, EReplicationBridgeDestroyInstanceFlags DestroyFlags,  UE::Net::FNetObjectFactoryId NetFactoryId)
-{
-}
-
 void UReplicationBridge::DetachInstance(FNetRefHandle Handle)
 {
 }
@@ -117,65 +106,6 @@ void UReplicationBridge::GetInitialDependencies(FNetRefHandle Handle, FNetDepend
 	return;
 }
 
-void UReplicationBridge::CallGetInitialDependencies(FNetRefHandle Handle, FNetDependencyInfoArray& OutDependencies) const
-{
-	using namespace UE::Net::Private;
-
-	// if the Handle is static, the initial dependency is the handle itself
-	if (Handle.IsStatic())
-	{
-		OutDependencies.Emplace(FObjectReferenceCache::MakeNetObjectReference(Handle));
-		return;
-	}
-	else
-	{
-		GetInitialDependencies(Handle, OutDependencies);
-	}
-}
-
-void UReplicationBridge::DetachSubObjectInstancesFromRemote(FNetRefHandle OwnerHandle, EReplicationBridgeDestroyInstanceReason DestroyReason, EReplicationBridgeDestroyInstanceFlags DestroyFlags)
-{
-	using namespace UE::Net::Private;
-
-	FInternalNetRefIndex OwnerInternalIndex = NetRefHandleManager->GetInternalIndex(OwnerHandle);
-	if (OwnerInternalIndex)
-	{
-		for (const FInternalNetRefIndex SubObjectInternalIndex : NetRefHandleManager->GetChildSubObjects(OwnerInternalIndex))
-		{
-			FNetRefHandleManager::FReplicatedObjectData& SubObjectData = NetRefHandleManager->GetReplicatedObjectDataNoCheck(SubObjectInternalIndex);
-			const FNetRefHandle SubObjectHandle = SubObjectData.RefHandle;
-			SubObjectData.bTearOff = (DestroyReason == EReplicationBridgeDestroyInstanceReason::TearOff);
-			SubObjectData.bPendingEndReplication = 1U;
-
-			EReplicationBridgeDestroyInstanceFlags SubObjectDestroyFlags = DestroyFlags;
-			// The subobject is allowed to be destroyed if both the owner and the subobject allows it.
-			if (!SubObjectData.bAllowDestroyInstanceFromRemote)
-			{
-				EnumRemoveFlags(SubObjectDestroyFlags, EReplicationBridgeDestroyInstanceFlags::AllowDestroyInstanceFromRemote);
-			}
-
-			if (DestroyReason == EReplicationBridgeDestroyInstanceReason::DoNotDestroy && SubObjectData.bAllowDestroyInstanceFromRemote)
-			{
-				// When ending replication of static objects without destroying them, we should always destroy instances spawned from replication
-				// as they will be recreated when the static object is scoped again.
-				if (bAlwaysDestroyDynamicSubObjectInstancesOnDetachFromStaticRoot)
-				{
-					EnumAddFlags(SubObjectDestroyFlags, EReplicationBridgeDestroyInstanceFlags::AllowDestroyInstanceFromRemote);
-					CallDetachInstanceFromRemote(SubObjectHandle, EReplicationBridgeDestroyInstanceReason::Destroy, SubObjectDestroyFlags, SubObjectData.NetFactoryId);
-					continue;
-				}
-				else
-				{
-					UE_LOG(LogIris, Warning, TEXT("Detaching Subobject %s with reason DoNotDestroy even though it is dynamic"), *SubObjectHandle.ToString());
-				}
-			}
-
-			CallDetachInstanceFromRemote(SubObjectHandle, DestroyReason, SubObjectDestroyFlags, SubObjectData.NetFactoryId);
-		}
-	}
-}
-
-
 void UReplicationBridge::CallPreSendUpdate(float DeltaSeconds)
 {
 	// Tear-off all handles pending tear-off
@@ -199,11 +129,6 @@ void UReplicationBridge::CallDetachInstance(FNetRefHandle Handle)
 	DetachInstance(Handle);
 }
 
-void UReplicationBridge::CallDetachInstanceFromRemote(FNetRefHandle Handle, EReplicationBridgeDestroyInstanceReason DestroyReason, EReplicationBridgeDestroyInstanceFlags DestroyFlags, UE::Net::FNetObjectFactoryId NetFactoryId)
-{
-	DetachInstanceFromRemote(Handle, DestroyReason, DestroyFlags, NetFactoryId);
-}
-
 void UReplicationBridge::CallPruneStaleObjects()
 {
 	PruneStaleObjects();
@@ -254,8 +179,6 @@ void UReplicationBridge::Deinitialize()
 		FNetRefHandleManager::FReplicatedObjectData& ObjectData = NetRefHandleManager->GetReplicatedObjectDataNoCheck(InternalObjectIndex);
 		if (ObjectData.InstanceProtocol)
 		{
-			const FNetRefHandle RefHandle = ObjectData.RefHandle;
-
 			// Clear out, tracking data or not?, currently we only have a single replication-system so it should not be a big issue if we do.
 			// Note: Currently we opted to leave it as is.
 			// -- This only applies to server. If it is a restart of ReplicationSystem, the same actors will be re-registered and global handle will be destroyed later, otherwise it does not matter what we do.
@@ -264,7 +187,7 @@ void UReplicationBridge::Deinitialize()
 
 			// Detach and destroy instance protocol
 			ObjectData.bPendingEndReplication = 1U;
-			InternalDetachInstanceFromNetRefHandle(ObjectData.RefHandle);
+			InternalDetachInstanceFromNetRefHandle(InternalObjectIndex);
 		}
 	});
 
@@ -321,20 +244,18 @@ void UReplicationBridge::InternalAttachInstanceToNetRefHandle(FNetRefHandle RefH
 	}
 }
 
-void UReplicationBridge::InternalDetachInstanceFromNetRefHandle(FNetRefHandle RefHandle)
+void UReplicationBridge::InternalDetachInstanceFromNetRefHandle(UE::Net::Private::FInternalNetRefIndex InternalObjectIndex)
 {
 	using namespace UE::Net;
 	using namespace UE::Net::Private;
 
-	const FInternalNetRefIndex InternalReplicationIndex = NetRefHandleManager->GetInternalIndex(RefHandle);
-
-	if (FReplicationInstanceProtocol* InstanceProtocol = const_cast<FReplicationInstanceProtocol*>(NetRefHandleManager->DetachInstanceProtocol(InternalReplicationIndex)))
+	if (FReplicationInstanceProtocol* InstanceProtocol = const_cast<FReplicationInstanceProtocol*>(NetRefHandleManager->DetachInstanceProtocol(InternalObjectIndex)))
 	{
-		UE_LOG_REPLICATIONBRIDGE(Verbose, TEXT("InternalDetachInstanceToNetHandle Detached: %s from (InternalIndex: %u)"), *RefHandle.ToString(), InternalReplicationIndex);
+		UE_LOG_REPLICATIONBRIDGE(Verbose, TEXT("InternalDetachInstanceToNetHandle Detached: %s"), *NetRefHandleManager->PrintObjectFromIndex(InternalObjectIndex));
 
 		if (EnumHasAnyFlags(InstanceProtocol->InstanceTraits, EReplicationInstanceProtocolTraits::IsBound))
 		{
-			FReplicationInstanceOperationsInternal::UnbindInstanceProtocol(InstanceProtocol, NetRefHandleManager->GetReplicatedObjectDataNoCheck(InternalReplicationIndex).Protocol);
+			FReplicationInstanceOperationsInternal::UnbindInstanceProtocol(InstanceProtocol, NetRefHandleManager->GetReplicatedObjectDataNoCheck(InternalObjectIndex).Protocol);
 		}
 		ReplicationProtocolManager->DestroyInstanceProtocol(InstanceProtocol);
 	}
@@ -385,10 +306,10 @@ void UReplicationBridge::DestroyLocalNetHandle(FNetRefHandle Handle, EEndReplica
 {
 	UE_LOG_REPLICATIONBRIDGE(Verbose, TEXT("DestroyLocalNetHandle for %s | EndReplicationFlags: %s"), *NetRefHandleManager->PrintObjectFromNetRefHandle(Handle), *LexToString(EndReplicationFlags));
 
+	const UE::Net::Private::FInternalNetRefIndex InternalReplicationIndex = NetRefHandleManager->GetInternalIndex(Handle);
+
 	if (EnumHasAnyFlags(EndReplicationFlags, EEndReplicationFlags::DestroyNetHandle | EEndReplicationFlags::ClearNetPushId))
 	{
-		const UE::Net::Private::FInternalNetRefIndex InternalReplicationIndex = NetRefHandleManager->GetInternalIndex(Handle);
-
 		if (EnumHasAnyFlags(EndReplicationFlags, EEndReplicationFlags::DestroyNetHandle))
 		{
 			DestroyGlobalNetHandle(InternalReplicationIndex);
@@ -401,7 +322,7 @@ void UReplicationBridge::DestroyLocalNetHandle(FNetRefHandle Handle, EEndReplica
 	}
 
 	// Detach instance protocol
-	InternalDetachInstanceFromNetRefHandle(Handle);
+	InternalDetachInstanceFromNetRefHandle(InternalReplicationIndex);
 
 	// Allow derived bridges to cleanup any instance info they have stored
 	CallDetachInstance(Handle);
@@ -505,7 +426,7 @@ void UReplicationBridge::StopReplicatingNetRefHandle(FNetRefHandle Handle, EEndR
 
 			// Detach instance as we must assume that we should not access the object after this call.
 			ObjectData.bPendingEndReplication = 1U;
-			InternalDetachInstanceFromNetRefHandle(Handle);
+			InternalDetachInstanceFromNetRefHandle(InternalReplicationIndex);
 		}
 		else 
 		{
@@ -527,7 +448,7 @@ void UReplicationBridge::StopReplicatingNetRefHandle(FNetRefHandle Handle, EEndR
 
 				// Detach instance as we must assume that we should not access the object after this call.
 				ObjectData.bPendingEndReplication = 1U;
-				InternalDetachInstanceFromNetRefHandle(Handle);
+				InternalDetachInstanceFromNetRefHandle(InternalReplicationIndex);
 			}
 			else
 			{
@@ -564,7 +485,7 @@ void UReplicationBridge::StopReplicatingNetRefHandle(FNetRefHandle Handle, EEndR
 			NetRefHandleManager->GetReplicatedObjectDataNoCheck(InternalReplicationIndex).bPendingEndReplication = 1U;
 		}
 		// If we get a call to end replication on the client, we need to detach the instance as it might be garbage collected
-		InternalDetachInstanceFromNetRefHandle(Handle);
+		InternalDetachInstanceFromNetRefHandle(InternalReplicationIndex);
 	}
 }
 
diff --git a/Engine/Source/Runtime/Net/Iris/Public/Iris/ReplicationSystem/ObjectReplicationBridge.h b/Engine/Source/Runtime/Net/Iris/Public/Iris/ReplicationSystem/ObjectReplicationBridge.h
--- a/Engine/Source/Runtime/Net/Iris/Public/Iris/ReplicationSystem/ObjectReplicationBridge.h
+++ b/Engine/Source/Runtime/Net/Iris/Public/Iris/ReplicationSystem/ObjectReplicationBridge.h
@@ -326,7 +326,6 @@ protected:
 	IRISCORE_API virtual void UpdateInstancesWorldLocation() override;
 	IRISCORE_API virtual void PruneStaleObjects() override;	
 	IRISCORE_API virtual bool CacheNetRefHandleCreationInfo(FNetRefHandle Handle) override;
-	IRISCORE_API virtual void DetachInstanceFromRemote(FNetRefHandle Handle, EReplicationBridgeDestroyInstanceReason DestroyReason, EReplicationBridgeDestroyInstanceFlags DestroyFlags,  UE::Net::FNetObjectFactoryId NetFactoryId) override;
 	IRISCORE_API virtual void DetachInstance(FNetRefHandle Handle) override;
 	IRISCORE_API virtual void OnProtocolMismatchReported(FNetRefHandle RefHandle, uint32 ConnectionId) override;
 	IRISCORE_API virtual void OnErrorWithNetRefHandleReported(UE::Net::ENetRefHandleError ErrorType, FNetRefHandle RefHandle, uint32 ConnectionId, const TArray<FNetRefHandle>& ExtraHandles) override;
@@ -451,6 +450,17 @@ private:
 
 	void RegisterRemoteInstance(FNetRefHandle RefHandle, UObject* InstancePtr, const UE::Net::FReplicationProtocol* Protocol, UE::Net::FReplicationInstanceProtocol* InstanceProtocol, uint32 ConnectionId);
 
+	/**
+	 * Called when the instance is detached from the protocol on request by the remote. 
+	 * @param Handle The handle of the object to destroy or tear off.
+	 * @param DestroyReason Reason for destroying the instance. 
+	 * @param DestroyFlags Special flags such as whether the instance may be destroyed when reason is TearOff, which may revert to destroying, or Destroy.
+	 */
+	void DetachInstanceFromRemote(FNetRefHandle Handle, EReplicationBridgeDestroyInstanceReason DestroyReason, EReplicationBridgeDestroyInstanceFlags DestroyFlags,  UE::Net::FNetObjectFactoryId NetFactoryId);
+
+	/** Detach a subobject on the client due to a request from the authority */
+	void DetachSubObjectInstancesFromRemote(UE::Net::Private::FInternalNetRefIndex InternalObjectIndex, EReplicationBridgeDestroyInstanceReason DestroyReason, EReplicationBridgeDestroyInstanceFlags DestroyFlags);
+
 	void SetNetPushIdOnInstance(UE::Net::FReplicationInstanceProtocol* InstanceProtocol, FNetHandle NetHandle);
 
 	/** Tries to load the classes used in poll period overrides. */
@@ -523,7 +533,7 @@ private:
 
 	/**
 	 * Client only function that executes the EndReplication request by the authority
-	 * @param Handle The netobject to destroy
+	 * @param InternalObjectIndex The index of the object to destroy
 	 * @DestroyReason How the object should be destroyed
 	 */
 	void DestroyNetObjectFromRemote(UE::Net::Private::FInternalNetRefIndex InternalObjectIndex, EReplicationBridgeDestroyInstanceReason DestroyReason);
diff --git a/Engine/Source/Runtime/Net/Iris/Public/Iris/ReplicationSystem/ReplicationBridge.h b/Engine/Source/Runtime/Net/Iris/Public/Iris/ReplicationSystem/ReplicationBridge.h
--- a/Engine/Source/Runtime/Net/Iris/Public/Iris/ReplicationSystem/ReplicationBridge.h
+++ b/Engine/Source/Runtime/Net/Iris/Public/Iris/ReplicationSystem/ReplicationBridge.h
@@ -184,14 +184,6 @@ protected:
 	*/
 	IRISCORE_API virtual bool CacheNetRefHandleCreationInfo(FNetRefHandle Handle);
 
-	/**
-	 * Called when the instance is detached from the protocol on request by the remote. 
-	 * @param Handle The handle of the object to destroy or tear off.
-	 * @param DestroyReason Reason for destroying the instance. 
-	 * @param DestroyFlags Special flags such as whether the instance may be destroyed when reason is TearOff, which may revert to destroying, or Destroy.
-	 */
-	IRISCORE_API virtual void DetachInstanceFromRemote(FNetRefHandle Handle, EReplicationBridgeDestroyInstanceReason DestroyReason, EReplicationBridgeDestroyInstanceFlags DestroyFlags,  UE::Net::FNetObjectFactoryId NetFactoryId);
-
 	/** Called when we detach instance protocol from the local instance */
 	IRISCORE_API virtual void DetachInstance(FNetRefHandle Handle);
 
@@ -224,7 +216,7 @@ protected:
 	IRISCORE_API void InternalAttachInstanceToNetRefHandle(FNetRefHandle RefHandle, bool bBindInstanceProtocol, UE::Net::FReplicationInstanceProtocol* InstanceProtocol, UObject* Instance, FNetHandle NetHandle);
 
 	/** Detach instance from NetRefHandle and destroy the instance protocol. */
-	IRISCORE_API void InternalDetachInstanceFromNetRefHandle(FNetRefHandle Handle);
+	IRISCORE_API void InternalDetachInstanceFromNetRefHandle(UE::Net::Private::FInternalNetRefIndex InternalObjectIndex);
 
 	/** Destroy the handle and all internal book keeping associated with it. */
 	IRISCORE_API void InternalDestroyNetObject(FNetRefHandle Handle);
@@ -271,8 +263,6 @@ private:
 
 	// Internal operations invoked by ReplicationSystem/ReplicationWriter
 	
-	void DetachSubObjectInstancesFromRemote(FNetRefHandle Handle, EReplicationBridgeDestroyInstanceReason DestroyReason, EReplicationBridgeDestroyInstanceFlags DestroyFlags);
-	
 	enum class EPendingEndReplicationImmediate : uint8
 	{
 		Yes,
@@ -283,15 +273,12 @@ private:
 	// it will be kept around until the handle is no longer ref-counted by any connection. It will however be removed from the set of scopeable objects after the first update so new connections will not add it to their scope.
 	void AddPendingEndReplication(FNetRefHandle Handle, EEndReplicationFlags DestroyFlags, EPendingEndReplicationImmediate Immediate = EPendingEndReplicationImmediate::No);
 
-	void CallPreSendUpdate(float DeltaSeconds);	
+	void CallPreSendUpdate(float DeltaSeconds);
 	void CallPreSendUpdateSingleHandle(FNetRefHandle Handle);
 	void CallUpdateInstancesWorldLocation();
 	bool CallCacheNetRefHandleCreationInfo(FNetRefHandle Handle);
-	bool CallWriteNetRefHandleCreationInfo(FReplicationBridgeSerializationContext& Context, FNetRefHandle Handle);
 	void CallPruneStaleObjects();
-	void CallGetInitialDependencies(FNetRefHandle Handle, FNetDependencyInfoArray& OutDependencies) const;
 	void CallDetachInstance(FNetRefHandle Handle);
-	void CallDetachInstanceFromRemote(FNetRefHandle Handle, EReplicationBridgeDestroyInstanceReason DestroyReason, EReplicationBridgeDestroyInstanceFlags DestroyFlags, UE::Net::FNetObjectFactoryId NetFactoryId);
 
 	void PreReceiveUpdate();
 	void PostReceiveUpdate();
