Iris
* Renamed NetObjectFactory::SubObjectDestroyedFromReplication to SubObjectDetachedFromReplication
* Clarified the function's documentation at the same time

#rnx

diff --git a/Engine/Plugins/Runtime/ReplicationSystemTestPlugin/Source/Private/Tests/ReplicationSystem/ReplicatedTestObjectFactory.cpp b/Engine/Plugins/Runtime/ReplicationSystemTestPlugin/Source/Private/Tests/ReplicationSystem/ReplicatedTestObjectFactory.cpp
--- a/Engine/Plugins/Runtime/ReplicationSystemTestPlugin/Source/Private/Tests/ReplicationSystem/ReplicatedTestObjectFactory.cpp
+++ b/Engine/Plugins/Runtime/ReplicationSystemTestPlugin/Source/Private/Tests/ReplicationSystem/ReplicatedTestObjectFactory.cpp
@@ -239,7 +239,7 @@ void UReplicatedTestObjectFactory::SubObjectCreatedFromReplication(UE::Net::FNet
 	}
 }
 
-void UReplicatedTestObjectFactory::SubObjectDestroyedFromReplication(const FDestroyedContext& Context)
+void UReplicatedTestObjectFactory::SubObjectDetachedFromReplication(const FDestroyedContext& Context)
 {
 	UReplicatedTestObject* RootInstance = Cast<UReplicatedTestObject>(Context.RootObject);
 	if (ensure(RootInstance))
diff --git a/Engine/Plugins/Runtime/ReplicationSystemTestPlugin/Source/Private/Tests/ReplicationSystem/ReplicatedTestObjectFactory.h b/Engine/Plugins/Runtime/ReplicationSystemTestPlugin/Source/Private/Tests/ReplicationSystem/ReplicatedTestObjectFactory.h
--- a/Engine/Plugins/Runtime/ReplicationSystemTestPlugin/Source/Private/Tests/ReplicationSystem/ReplicatedTestObjectFactory.h
+++ b/Engine/Plugins/Runtime/ReplicationSystemTestPlugin/Source/Private/Tests/ReplicationSystem/ReplicatedTestObjectFactory.h
@@ -44,7 +44,7 @@ public:
 	virtual void DestroyReplicatedObject(const FDestroyedContext& Context) override;
 	virtual TOptional<FWorldInfoData> GetWorldInfo(const FWorldInfoContext& Context) const override;
 	virtual void SubObjectCreatedFromReplication(UE::Net::FNetRefHandle RootObject, UE::Net::FNetRefHandle SubObjectCreated) override;
-	virtual void SubObjectDestroyedFromReplication(const FDestroyedContext& Context) override;
+	virtual void SubObjectDetachedFromReplication(const FDestroyedContext& Context) override;
 
 	//------------------------------------------------------------------------
 
diff --git a/Engine/Source/Runtime/Net/Iris/Private/Iris/ReplicationSystem/ObjectReplicationBridge.cpp b/Engine/Source/Runtime/Net/Iris/Private/Iris/ReplicationSystem/ObjectReplicationBridge.cpp
--- a/Engine/Source/Runtime/Net/Iris/Private/Iris/ReplicationSystem/ObjectReplicationBridge.cpp
+++ b/Engine/Source/Runtime/Net/Iris/Private/Iris/ReplicationSystem/ObjectReplicationBridge.cpp
@@ -1189,7 +1189,7 @@ void UObjectReplicationBridge::DetachInstanceFromRemote(FNetRefHandle Handle, ER
 				DestroyContext.RootObject = NetRefHandleManager->GetReplicatedObjectInstance(DestroyedObjectData.SubObjectRootIndex);
 
 				UNetObjectFactory* RootObjectFactory = GetNetFactory(NetRefHandleManager->GetReplicatedObjectData(DestroyedObjectData.SubObjectRootIndex).NetFactoryId);
-				RootObjectFactory->SubObjectDestroyedFromReplication(DestroyContext);
+				RootObjectFactory->SubObjectDetachedFromReplication(DestroyContext);
 			}			
 
 			// Now ask the object's factory to handle the destruction
diff --git a/Engine/Source/Runtime/Net/Iris/Public/Iris/ReplicationSystem/NetObjectFactory.h b/Engine/Source/Runtime/Net/Iris/Public/Iris/ReplicationSystem/NetObjectFactory.h
--- a/Engine/Source/Runtime/Net/Iris/Public/Iris/ReplicationSystem/NetObjectFactory.h
+++ b/Engine/Source/Runtime/Net/Iris/Public/Iris/ReplicationSystem/NetObjectFactory.h
@@ -157,11 +157,12 @@ public:
 	virtual void SubObjectCreatedFromReplication(UE::Net::FNetRefHandle RootObject, UE::Net::FNetRefHandle SubObjectCreated) {}
 
 	/**
-	 * Optional callback triggered when a subobject is about to destroyed for the owner root object managed by this factory.
-	 * This callback is called on remotes only.
+	 * Optional callback triggered when a subobject will be detached from the replication system and potentially destroyed
+	 * by it's factory in DestroyReplicatedObject. Both static and dynamic subobjects will be passed to this function.
+	 * Note that this callback is called on remotes (clients) only.
 	 * @param Context Gives access to the subobject about to be destroyed and the root object that owns him.
 	 */
-	virtual void SubObjectDestroyedFromReplication(const FDestroyedContext& Context) {}
+	virtual void SubObjectDetachedFromReplication(const FDestroyedContext& Context) {}
 
 	/**
 	 * Fetch world information about a replicated object so it can be updated in the network engine.
