Iris
* Added FindByName support to Net.Iris.PrintNetInfoOfObject (formally PrintDebugInfoOfNetHandle)

#rnx

diff --git a/Engine/Source/Runtime/Engine/Private/Tests/NetTestHelpers.cpp b/Engine/Source/Runtime/Engine/Private/Tests/NetTestHelpers.cpp
--- a/Engine/Source/Runtime/Engine/Private/Tests/NetTestHelpers.cpp
+++ b/Engine/Source/Runtime/Engine/Private/Tests/NetTestHelpers.cpp
@@ -385,6 +385,13 @@ bool FTestWorlds::CreateAndConnectClient()
 	return bConnected;
 }
 
+int32 FTestWorlds::ConnectNewClient()
+{
+	Clients.Emplace(FTestWorldInstance::CreateClient(Server.GetPort()));
+	const bool bConnected = WaitForClientConnect(Clients.Last());
+	return bConnected ? (Clients.Num()-1) : INDEX_NONE;
+}
+
 bool FTestWorlds::WaitForClientConnect(FTestWorldInstance& Client)
 {
 	bool bConnected = TickAllUntil([&Client]()
diff --git a/Engine/Source/Runtime/Engine/Public/Tests/NetTestHelpers.h b/Engine/Source/Runtime/Engine/Public/Tests/NetTestHelpers.h
--- a/Engine/Source/Runtime/Engine/Public/Tests/NetTestHelpers.h
+++ b/Engine/Source/Runtime/Engine/Public/Tests/NetTestHelpers.h
@@ -200,9 +200,12 @@ struct FTestWorlds
 	ENGINE_API void SetTickInSeconds(float TickInSeconds);
 	float GetTickDeltaSeconds() const { return TickDeltaSeconds; }
 	
-	/** Create a client and fully connect it to the server before returning */
+	/** Create a client and fully connect it to the server. Returns true if the connection succeeded */
 	ENGINE_API bool CreateAndConnectClient();
 
+	/** Create a new client and connect it to the server. Returns the ClientID of the new client or INDEX_NONE if it failed to connect. */
+	ENGINE_API int32 ConnectNewClient();
+
 	/** Ticks all server & client worlds NumTick times synchronously. */
 	ENGINE_API void TickAll(int32 NumTicks=1);
 	ENGINE_API void TickServer();
diff --git a/Engine/Source/Runtime/Net/Iris/Private/Iris/ReplicationSystem/ObjectReplicationBridge.cpp b/Engine/Source/Runtime/Net/Iris/Private/Iris/ReplicationSystem/ObjectReplicationBridge.cpp
--- a/Engine/Source/Runtime/Net/Iris/Private/Iris/ReplicationSystem/ObjectReplicationBridge.cpp
+++ b/Engine/Source/Runtime/Net/Iris/Private/Iris/ReplicationSystem/ObjectReplicationBridge.cpp
@@ -3051,7 +3051,7 @@ void UObjectReplicationBridge::OnErrorCannotReplicateObject(FNetRefHandle RefHan
 	}
 	else
 	{
-		UE_LOG_BRIDGEID(LogIris, Error, TEXT("Client cannot replicated object: %s. Reported by client: %s"), *NetRefHandleManager->PrintObjectFromIndex(ObjectInternalIndex), *PrintConnectionInfo(ConnectionId));
+		UE_LOG_BRIDGEID(LogIris, Error, TEXT("Client cannot replicate object: %s. Reported by client: %s"), *NetRefHandleManager->PrintObjectFromIndex(ObjectInternalIndex), *PrintConnectionInfo(ConnectionId));
 	}
 
 	ensureMsgf(false, TEXT("Client disabled replication for %s. Client was unable to read the creation data and did not spawn the netobject."), *NetRefHandleManager->PrintObjectFromNetRefHandle(RefHandle));
diff --git a/Engine/Source/Runtime/Net/Iris/Private/Iris/ReplicationSystem/ObjectReplicationBridgeDebugging.cpp b/Engine/Source/Runtime/Net/Iris/Private/Iris/ReplicationSystem/ObjectReplicationBridgeDebugging.cpp
--- a/Engine/Source/Runtime/Net/Iris/Private/Iris/ReplicationSystem/ObjectReplicationBridgeDebugging.cpp
+++ b/Engine/Source/Runtime/Net/Iris/Private/Iris/ReplicationSystem/ObjectReplicationBridgeDebugging.cpp
@@ -318,7 +318,8 @@ void LogViaTrait(FLogContext& LogContext, EPrintTraits ArgTraits, EPrintTraits D
  */
 TArray<FNetRefHandle> FindNetRefHandlesFromArg(UReplicationSystem* RepSystem, const TArray<FString>& Args)
 {
-	//TODO: Add more ways to find objects FindByInternalIndex, FindByPtr, FindByName, etc.
+	//TODO: Add more ways to find objects FindByInternalIndex, FindByPtr, etc.
+
 	TArray<FNetRefHandle> NetRefHandles;
 	 
 	const FNetRefHandleManager&  NetRefHandleManager = RepSystem->GetReplicationSystemInternal()->GetNetRefHandleManager();
@@ -347,6 +348,44 @@ TArray<FNetRefHandle> FindNetRefHandlesFromArg(UReplicationSystem* RepSystem, co
 			}
 		}
 	}
+
+	if (const FString* FindByName = Args.FindByPredicate([](const FString& Str) { return Str.Contains(TEXT("FindByName=")); }))
+	{
+		FString ObjectNameToFind;
+		FParse::Value(**FindByName, TEXT("FindByName="), ObjectNameToFind);
+
+		TArray<UObject*> FoundObjects;
+
+		for (FThreadSafeObjectIterator It; It; ++It)
+		{
+			UObject* Object = *It;
+
+			if (!IsValid(Object) || Object->HasAnyFlags(RF_ArchetypeObject | RF_ClassDefaultObject))
+			{
+				continue;
+			}
+
+			const FString FullName = Object->GetName();
+			if (FullName.Contains(ObjectNameToFind))
+			{
+				FoundObjects.Emplace(Object);
+			}
+		}
+
+		// Add the NetRefHandles of replicated objects matching the name
+		for (UObject* FoundObject : FoundObjects)
+		{
+			const FNetRefHandle NetHandle = RepSystem->GetReplicationBridge()->GetReplicatedRefHandle(FoundObject);
+			if (NetHandle.IsValid())
+			{
+				NetRefHandles.Emplace(NetHandle);
+			}
+			else
+			{
+				// TODO: add a flag to export objects that are not replicated but name match
+			}
+		}
+	}
 	
 	if (const FString* FindById = Args.FindByPredicate([](const FString& Str) { return Str.Contains(TEXT("FindById=")); }))
 	{
@@ -390,7 +429,7 @@ FAutoConsoleCommand ObjectBridgePrintDynamicFilterConfig(
 		return;
 	}
 
-	UObjectReplicationBridge* ObjectBridge = CastChecked<UObjectReplicationBridge>(RepSystem->GetReplicationBridge());
+	UObjectReplicationBridge* ObjectBridge = RepSystem->GetReplicationBridge();
 	if (!ObjectBridge)
 	{
 		UE_LOG(LogIris, Error, TEXT("Could not find ObjectReplicationBridge."));
@@ -415,7 +454,7 @@ FAutoConsoleCommand ObjectBridgePrintDynamicFilterConfigIssues(
 		return;
 	}
 
-	UObjectReplicationBridge* ObjectBridge = CastChecked<UObjectReplicationBridge>(RepSystem->GetReplicationBridge());
+	UObjectReplicationBridge* ObjectBridge = RepSystem->GetReplicationBridge();
 	if (!ObjectBridge)
 	{
 		UE_LOG(LogIris, Error, TEXT("Could not find ObjectReplicationBridge."));
@@ -640,7 +679,7 @@ FAutoConsoleCommand ObjectBridgePrintReplicatedObjects(
 	UReplicationSystem* RepSystem = FindReplicationSystemFromArg(Args);
 	if (RepSystem)
 	{
-		if (UObjectReplicationBridge* ObjectBridge = CastChecked<UObjectReplicationBridge>(RepSystem->GetReplicationBridge()))
+		if (UObjectReplicationBridge* ObjectBridge = RepSystem->GetReplicationBridge())
 		{
 			EPrintTraits ArgTraits = FindPrintTraitsFromArgs(Args);
 			ObjectBridge->PrintReplicatedObjects((uint32)ArgTraits);
@@ -697,7 +736,7 @@ FAutoConsoleCommand ObjectBridgePrintRelevantObjects(
 
 	if (UReplicationSystem* RepSystem = FindReplicationSystemFromArg(Args))
 	{
-		if (UObjectReplicationBridge* ObjectBridge = CastChecked<UObjectReplicationBridge>(RepSystem->GetReplicationBridge()))
+		if (UObjectReplicationBridge* ObjectBridge = RepSystem->GetReplicationBridge())
 		{
 			EPrintTraits ArgTraits = FindPrintTraitsFromArgs(Args);
 			ObjectBridge->PrintRelevantObjects((uint32)ArgTraits);
@@ -743,7 +782,7 @@ FAutoConsoleCommand ObjectBridgePrintAlwaysRelevantObjects(
 
 	if (UReplicationSystem* RepSystem = FindReplicationSystemFromArg(Args))
 	{
-		if (UObjectReplicationBridge* ObjectBridge = CastChecked<UObjectReplicationBridge>(RepSystem->GetReplicationBridge()))
+		if (UObjectReplicationBridge* ObjectBridge = RepSystem->GetReplicationBridge())
 		{
 			EPrintTraits ArgTraits = FindPrintTraitsFromArgs(Args);
 			ObjectBridge->PrintAlwaysRelevantObjects((uint32)ArgTraits);
@@ -815,7 +854,7 @@ FConsoleCommandWithArgsDelegate::CreateLambda([](const TArray< FString >& Args)
 	UReplicationSystem* RepSystem = FindReplicationSystemFromArg(Args);
 	if (RepSystem)
 	{
-		if (UObjectReplicationBridge* ObjectBridge = CastChecked<UObjectReplicationBridge>(RepSystem->GetReplicationBridge()))
+		if (UObjectReplicationBridge* ObjectBridge = RepSystem->GetReplicationBridge())
 		{
 			FReplicationSystemInternal* ReplicationSystemInternal = RepSystem->GetReplicationSystemInternal();
 
@@ -920,7 +959,7 @@ FConsoleCommandWithArgsDelegate::CreateLambda([](const TArray< FString >& Args)
 	UReplicationSystem* RepSystem = FindReplicationSystemFromArg(Args);
 	if (RepSystem)
 	{
-		if (UObjectReplicationBridge* ObjectBridge = CastChecked<UObjectReplicationBridge>(RepSystem->GetReplicationBridge()))
+		if (UObjectReplicationBridge* ObjectBridge = RepSystem->GetReplicationBridge())
 		{
 			ObjectBridge->PrintNetCullDistances(Args);
 		}
@@ -1076,7 +1115,7 @@ FAutoConsoleCommand ObjectBridgePrintPushBasedStatuses(
 		return;
 	}
 
-	UObjectReplicationBridge* ObjectBridge = RepSystem->GetReplicationBridgeAs<UObjectReplicationBridge>();
+	UObjectReplicationBridge* ObjectBridge = RepSystem->GetReplicationBridge();
 	if (!ObjectBridge)
 	{
 		UE_LOG(LogIris, Error, TEXT("Could not find ObjectReplicationBridge."));
@@ -1170,9 +1209,18 @@ void UObjectReplicationBridge::PrintPushBasedStatuses() const
 }
 
 //-----------------------------------------------
-FAutoConsoleCommand ObjectBridgePrintDebugInfoOfNetHandle(
-	TEXT("Net.Iris.PrintDebugInfoOfNetHandle"),
-	TEXT("Prints all relevant information about a NetHandle and a connection."),
+FAutoConsoleCommand ObjectBridgePrintNetInfoOfObject(
+	TEXT("Net.Iris.PrintNetInfoOfObject"),
+	TEXT("Prints all relevant information about replicated objects and their relevancy status for clients in the session."),
+/**
+ * Search params
+ * FindByName=<object name> : Wildcard match for replicated UObject name.
+ * FindByClass=<class name> : Wildcard match the class names (with inheritance).
+ * FindById=<id list> : Find the objects via their NetRefHandle Id. ex: '(Id=2010):(RepSystemId=0)' would be 'FindById=2010'
+ * 
+ * Connection params
+ * ConnectionId=<id list> : Filter relevancy status of specific connections
+ */
 FConsoleCommandWithArgsDelegate::CreateLambda([](const TArray< FString >& Args)
 {
 	using namespace UE::Net;
@@ -1186,7 +1234,7 @@ FConsoleCommandWithArgsDelegate::CreateLambda([](const TArray< FString >& Args)
 		return;
 	}
 
-	UObjectReplicationBridge* ObjectBridge = RepSystem->GetReplicationBridgeAs<UObjectReplicationBridge>();
+	UObjectReplicationBridge* ObjectBridge = RepSystem->GetReplicationBridge();
 	if (!ObjectBridge)
 	{
 		UE_LOG(LogIris, Error, TEXT("Could not find ObjectReplicationBridge."));
@@ -1212,7 +1260,7 @@ void UObjectReplicationBridge::PrintDebugInfoForNetRefHandlesAndConnections(cons
 	// Default to all connections
 	FNetBitArray ConnectionsToPrint(ValidConnections);
 
-	// Filter down the list if specific connections were specified
+	// Filter down the list of specific connections if requested
 	{
 		const TArray<uint32> RequestedConnectionList = FindConnectionsFromArgs(Args);
 		if (!RequestedConnectionList.IsEmpty())
@@ -1259,8 +1307,6 @@ void UObjectReplicationBridge::PrintDebugInfoForNetRefHandle(const FNetRefHandle
 
 	UE_LOG(LogIris, Display, TEXT("################ Start Printing debug info for: %s ################"), *NetRefHandleManager->PrintObjectFromNetRefHandle(NetRefHandle));
 	UE_LOG(LogIris, Display, TEXT(""));
-	UE_LOG(LogIris, Display, TEXT("For Connection: %s"), *PrintConnectionInfo(ConnectionId))
-	UE_LOG(LogIris, Display, TEXT(""));
 
 #if !UE_BUILD_SHIPPING
 	if (!EnumHasAnyFlags(PrintTraits, EPrintDebugInfoTraits::NoProtocolState))
@@ -1273,6 +1319,10 @@ void UObjectReplicationBridge::PrintDebugInfoForNetRefHandle(const FNetRefHandle
 	}
 #endif
 
+	UE_LOG(LogIris, Display, TEXT(""));
+	UE_LOG(LogIris, Display, TEXT("For Connection: %s"), *PrintConnectionInfo(ConnectionId))
+	UE_LOG(LogIris, Display, TEXT(""));
+
 	if (GetReplicationSystem()->IsServer())
 	{
 		// Connection info
