Workaround static analysis issue.

#rnx

diff --git a/Engine/Source/Runtime/Net/Iris/Private/Iris/Serialization/ArrayPropertyNetSerializer.cpp b/Engine/Source/Runtime/Net/Iris/Private/Iris/Serialization/ArrayPropertyNetSerializer.cpp
--- a/Engine/Source/Runtime/Net/Iris/Private/Iris/Serialization/ArrayPropertyNetSerializer.cpp
+++ b/Engine/Source/Runtime/Net/Iris/Private/Iris/Serialization/ArrayPropertyNetSerializer.cpp
@@ -310,7 +310,7 @@ void FArrayPropertyNetSerializer::SerializeDelta(FNetSerializationContext& Conte
 		ElementArgs.Source = NetSerializerValuePointer(Array.ElementStorage);
 		ElementArgs.Prev = NetSerializerValuePointer(PrevArray.ElementStorage);
 
-		if (!ChangeMask)
+		if (!ChangeMaskBitCount)
 		{
 			for (uint32 ElementIt = 0, ElementEndIt = FMath::Min(ElementCount, PrevElementCount); ElementIt < ElementEndIt; ++ElementIt)
 			{
@@ -343,7 +343,7 @@ void FArrayPropertyNetSerializer::SerializeDelta(FNetSerializationContext& Conte
 		ElementArgs.NetSerializerConfig = ElementSerializerDescriptor.SerializerConfig;
 		ElementArgs.Source = NetSerializerValuePointer(Array.ElementStorage) + PrevElementCount*ElementSize;
 
-		if (!ChangeMask)
+		if (!ChangeMaskBitCount)
 		{
 			for (uint32 ElementIt = PrevElementCount, ElementEndIt = ElementCount; ElementIt < ElementEndIt; ++ElementIt)
 			{
@@ -410,7 +410,7 @@ void FArrayPropertyNetSerializer::DeserializeDelta(FNetSerializationContext& Con
 		ElementArgs.Target = NetSerializerValuePointer(Array.ElementStorage);
 		ElementArgs.Prev = NetSerializerValuePointer(PrevArray.ElementStorage);
 
-		if (!ChangeMask)
+		if (!ChangeMaskBitCount)
 		{
 			for (uint32 ElementIt = 0, ElementEndIt = FMath::Min(ElementCount, PrevElementCount); ElementIt < ElementEndIt; ++ElementIt)
 			{
@@ -443,7 +443,7 @@ void FArrayPropertyNetSerializer::DeserializeDelta(FNetSerializationContext& Con
 		ElementArgs.NetSerializerConfig = ElementSerializerDescriptor.SerializerConfig;
 		ElementArgs.Target = NetSerializerValuePointer(Array.ElementStorage) + PrevElementCount*ElementSize;
 
-		if (!ChangeMask)
+		if (!ChangeMaskBitCount)
 		{
 			for (uint32 ElementIt = PrevElementCount, ElementEndIt = ElementCount; ElementIt < ElementEndIt; ++ElementIt)
 			{
