Iris: Allow external modules to log CsvStats under the Iris category

diff --git a/Engine/Source/Runtime/Net/Iris/Public/Iris/Core/IrisCsv.h b/Engine/Source/Runtime/Net/Iris/Public/Iris/Core/IrisCsv.h
--- a/Engine/Source/Runtime/Net/Iris/Public/Iris/Core/IrisCsv.h
+++ b/Engine/Source/Runtime/Net/Iris/Public/Iris/Core/IrisCsv.h
@@ -6,7 +6,7 @@
 #include "Iris/IrisConfig.h"
 #include "Iris/Core/IrisProfiler.h"
 
-CSV_DECLARE_CATEGORY_EXTERN(Iris);
+CSV_DECLARE_CATEGORY_MODULE_EXTERN(IRISCORE_API, Iris);
 
 #if UE_NET_IRIS_CSV_STATS
 	#define IRIS_CSV_PROFILER_SCOPE(CsvCategory, x) \
