Iris
* Added GetRootObjectOfAnyObject. Used to get a root object handle without knowing if the current object is a subobject or root object
* Renamed existing FReplicatedObjectData::IsSubObject to IsOwnedSubobject because it checked the bDestroySubObjectWithOwner flag (!)
* Added FReplicatedObjectData::IsSubObject that is true for all types of subobjects

#rnx

diff --git a/Engine/Source/Runtime/Net/Iris/Private/Iris/ReplicationSystem/Conditionals/ReplicationConditionals.cpp b/Engine/Source/Runtime/Net/Iris/Private/Iris/ReplicationSystem/Conditionals/ReplicationConditionals.cpp
--- a/Engine/Source/Runtime/Net/Iris/Private/Iris/ReplicationSystem/Conditionals/ReplicationConditionals.cpp
+++ b/Engine/Source/Runtime/Net/Iris/Private/Iris/ReplicationSystem/Conditionals/ReplicationConditionals.cpp
@@ -1052,7 +1052,7 @@ void FReplicationConditionals::InvalidateBaselinesForObjectHierarchy(uint32 Obje
 	for (const FInternalNetRefIndex SubObjectIndex : NetRefHandleManager->GetSubObjects(ObjectIndex))
 	{
 		const FNetRefHandleManager::FReplicatedObjectData& SubObjectData = NetRefHandleManager->GetReplicatedObjectDataNoCheck(SubObjectIndex);
-		if (SubObjectData.IsSubObject())
+		if (SubObjectData.IsOwnedSubObject())
 		{
 			if (EnumHasAnyFlags(SubObjectData.Protocol->ProtocolTraits, EReplicationProtocolTraits::HasLifetimeConditionals))
 			{
diff --git a/Engine/Source/Runtime/Net/Iris/Private/Iris/ReplicationSystem/Filtering/ReplicationFiltering.cpp b/Engine/Source/Runtime/Net/Iris/Private/Iris/ReplicationSystem/Filtering/ReplicationFiltering.cpp
--- a/Engine/Source/Runtime/Net/Iris/Private/Iris/ReplicationSystem/Filtering/ReplicationFiltering.cpp
+++ b/Engine/Source/Runtime/Net/Iris/Private/Iris/ReplicationSystem/Filtering/ReplicationFiltering.cpp
@@ -1066,7 +1066,7 @@ uint32 FReplicationFiltering::GetOwningConnectionIfDirty(uint32 ObjectIndex) con
 	if (NetRefHandleManager->GetSubObjectInternalIndices().GetBit(ObjectIndex))
 	{
 		const FNetRefHandleManager::FReplicatedObjectData& ObjectData = NetRefHandleManager->GetReplicatedObjectDataNoCheck(ObjectIndex);
-		if (ObjectData.IsSubObject())
+		if (ObjectData.IsOwnedSubObject())
 		{
 			const uint32 ParentIndex = NetRefHandleManager->GetReplicatedObjectDataNoCheck(ObjectIndex).SubObjectRootIndex;
 			return ObjectIndexToOwningConnection[ParentIndex];
@@ -1095,7 +1095,7 @@ void FReplicationFiltering::UpdateOwnerFiltering()
 			for (const FInternalNetRefIndex SubObjectIndex : NetRefHandleManager->GetSubObjects(ObjectIndex))
 			{
 				const FNetRefHandleManager::FReplicatedObjectData& ObjectData = NetRefHandleManager->GetReplicatedObjectDataNoCheck(SubObjectIndex);
-				if (ObjectData.IsSubObject())
+				if (ObjectData.IsOwnedSubObject())
 				{
 					ObjectIndexToOwningConnection[SubObjectIndex] = static_cast<uint16>(OwningConnectionId);
 				}
diff --git a/Engine/Source/Runtime/Net/Iris/Private/Iris/ReplicationSystem/NetRefHandleManager.cpp b/Engine/Source/Runtime/Net/Iris/Private/Iris/ReplicationSystem/NetRefHandleManager.cpp
--- a/Engine/Source/Runtime/Net/Iris/Private/Iris/ReplicationSystem/NetRefHandleManager.cpp
+++ b/Engine/Source/Runtime/Net/Iris/Private/Iris/ReplicationSystem/NetRefHandleManager.cpp
@@ -31,6 +31,11 @@ CSV_DEFINE_CATEGORY(IrisCommon, true);
 namespace UE::Net::Private
 {
 
+static bool IsValidIndex(FInternalNetRefIndex InternalIndex)
+{
+	return InternalIndex != FNetRefHandleManager::InvalidInternalIndex;
+}
+
 FNetRefHandleManager::FNetRefHandleManager(FReplicationProtocolManager& InReplicationProtocolManager)
 	: ReplicationProtocolManager(InReplicationProtocolManager)
 {
@@ -796,12 +801,26 @@ bool FNetRefHandleManager::SetSubObjectNetCondition(FInternalNetRefIndex SubObje
 	return false;
 }
 
+FNetRefHandle FNetRefHandleManager::GetRootObjectOfAnyObject(FNetRefHandle NetRefHandle) const
+{
+	const FInternalNetRefIndex InternalIndex = GetInternalIndex(NetRefHandle);
+	if (!IsValidIndex(InternalIndex))
+	{
+		return FNetRefHandle();
+	}
+
+	// Find the rootobject for subojects, otherwise just use the passed index since its a root.
+	const FInternalNetRefIndex RootObjectInternalIndex = ReplicatedObjectData[InternalIndex].IsSubObject() ? ReplicatedObjectData[InternalIndex].SubObjectRootIndex : InternalIndex;
+
+	return IsValidIndex(RootObjectInternalIndex) ? ReplicatedObjectData[RootObjectInternalIndex].RefHandle : FNetRefHandle();
+}
+
 FNetRefHandle FNetRefHandleManager::GetRootObjectOfSubObject(FNetRefHandle SubObjectRefHandle) const
 {
 	const FInternalNetRefIndex SubObjectInternalIndex = GetInternalIndex(SubObjectRefHandle);
-	const FInternalNetRefIndex OwnerInternalIndex = SubObjectInternalIndex != InvalidInternalIndex ? ReplicatedObjectData[SubObjectInternalIndex].SubObjectRootIndex : InvalidInternalIndex;
+	const FInternalNetRefIndex OwnerInternalIndex = IsValidIndex(SubObjectInternalIndex) ? ReplicatedObjectData[SubObjectInternalIndex].SubObjectRootIndex : InvalidInternalIndex;
 
-	return OwnerInternalIndex != InvalidInternalIndex ? ReplicatedObjectData[OwnerInternalIndex].RefHandle : FNetRefHandle();
+	return IsValidIndex(OwnerInternalIndex) ? ReplicatedObjectData[OwnerInternalIndex].RefHandle : FNetRefHandle();
 }
 
 FInternalNetRefIndex FNetRefHandleManager::GetRootObjectInternalIndexOfSubObject(FInternalNetRefIndex SubObjectIndex) const
diff --git a/Engine/Source/Runtime/Net/Iris/Private/Iris/ReplicationSystem/NetRefHandleManager.h b/Engine/Source/Runtime/Net/Iris/Private/Iris/ReplicationSystem/NetRefHandleManager.h
--- a/Engine/Source/Runtime/Net/Iris/Private/Iris/ReplicationSystem/NetRefHandleManager.h
+++ b/Engine/Source/Runtime/Net/Iris/Private/Iris/ReplicationSystem/NetRefHandleManager.h
@@ -120,10 +120,13 @@ public:
 			};
 		};
 	
-		// Returns true if this is a SubObject
-		bool IsSubObject() const { return SubObjectRootIndex && !!bDestroySubObjectWithOwner; }
+		/** Returns true if this is a SubObject that gets destroyed when his RootObject is destroyed too. */
+		bool IsOwnedSubObject() const { return SubObjectRootIndex != InvalidInternalIndex && !!bDestroySubObjectWithOwner; }
 
-		// Returns true if this is a DependentObject
+		/** Returns true if this is a SubObject */
+		bool IsSubObject() const { return SubObjectRootIndex != InvalidInternalIndex; }
+
+		/** Returns true if this is a child dependent object  */
 		bool IsDependentObject() const { return bIsDependentObject; }
 	};
 
@@ -258,6 +261,7 @@ public:
 	bool AddSubObject(FNetRefHandle RootObjectHandle, FNetRefHandle SubObjectHandle, EAddSubObjectFlags Flags = EAddSubObjectFlags::Default);
 	void RemoveSubObject(FNetRefHandle SubObjectHandle);
 	
+	FNetRefHandle GetRootObjectOfAnyObject(FNetRefHandle NetRefHandle) const;
 	FNetRefHandle GetRootObjectOfSubObject(FNetRefHandle SubObjectHandle) const;
 	FInternalNetRefIndex GetRootObjectInternalIndexOfSubObject(FInternalNetRefIndex SubObjectIndex) const;
 
diff --git a/Engine/Source/Runtime/Net/Iris/Private/Iris/ReplicationSystem/ObjectReplicationBridge.cpp b/Engine/Source/Runtime/Net/Iris/Private/Iris/ReplicationSystem/ObjectReplicationBridge.cpp
--- a/Engine/Source/Runtime/Net/Iris/Private/Iris/ReplicationSystem/ObjectReplicationBridge.cpp
+++ b/Engine/Source/Runtime/Net/Iris/Private/Iris/ReplicationSystem/ObjectReplicationBridge.cpp
@@ -990,14 +990,14 @@ void UObjectReplicationBridge::SetSubObjectNetCondition(FNetRefHandle SubObjectR
 	}
 }
 
-UE::Net::FNetRefHandle UObjectReplicationBridge::GetRootObjectOfSubObject(FNetRefHandle SubObjectHandle) const
+UE::Net::FNetRefHandle UObjectReplicationBridge::GetRootObjectOfAnyObject(FNetRefHandle NetRefHandle) const
 {
-	using namespace UE::Net::Private;
-
-	FReplicationSystemInternal* ReplicationSystemInternal = GetReplicationSystem()->GetReplicationSystemInternal();
-	FNetRefHandleManager& LocalNetRefHandleManager = ReplicationSystemInternal->GetNetRefHandleManager();
+	return GetReplicationSystem()->GetReplicationSystemInternal()->GetNetRefHandleManager().GetRootObjectOfAnyObject(NetRefHandle);
+}
 
-	return LocalNetRefHandleManager.GetRootObjectOfSubObject(SubObjectHandle);
+UE::Net::FNetRefHandle UObjectReplicationBridge::GetRootObjectOfSubObject(FNetRefHandle SubObjectHandle) const
+{
+	return GetReplicationSystem()->GetReplicationSystemInternal()->GetNetRefHandleManager().GetRootObjectOfSubObject(SubObjectHandle);
 }
 
 
diff --git a/Engine/Source/Runtime/Net/Iris/Private/Iris/ReplicationSystem/ReplicationWriter.cpp b/Engine/Source/Runtime/Net/Iris/Private/Iris/ReplicationSystem/ReplicationWriter.cpp
--- a/Engine/Source/Runtime/Net/Iris/Private/Iris/ReplicationSystem/ReplicationWriter.cpp
+++ b/Engine/Source/Runtime/Net/Iris/Private/Iris/ReplicationSystem/ReplicationWriter.cpp
@@ -708,7 +708,7 @@ void FReplicationWriter::SetPendingDestroyOrSubObjectPendingDestroyState(uint32
 	{
 		// Subobject destroyed before its owner is explicitly replicated as state data.
 		FNetRefHandleManager::FReplicatedObjectData& ObjectData = NetRefHandleManager->GetReplicatedObjectDataNoCheck(InternalIndex);
-		if (ObjectData.IsSubObject())
+		if (ObjectData.IsOwnedSubObject())
 		{
 			// If owner is not pending destroy we mark the state of the SubObject to SubObjectPendingDestroy and mark owner as having dirty subobjects which will 
 			// destroy the subobject using the replicated state path of the owner.
@@ -961,7 +961,7 @@ void FReplicationWriter::UpdateDirtyGlobalLifetimeConditionals(TArrayView<FInter
 			if (Info.IsSubObject)
 			{
 				const FNetRefHandleManager::FReplicatedObjectData& ObjectData = NetRefHandleManager->GetReplicatedObjectDataNoCheck(InternalObjectIndex);
-				if (ObjectData.IsSubObject())
+				if (ObjectData.IsOwnedSubObject())
 				{
 					const FInternalNetRefIndex RootInternalObjectIndex = ObjectData.SubObjectRootIndex;
 					FReplicationInfo& OwnerInfo = ReplicatedObjects[RootInternalObjectIndex];
@@ -1250,7 +1250,7 @@ void FReplicationWriter::HandleDeliveredRecord(const FReplicationRecord::FRecord
 					if (Info.IsSubObject)
 					{
 						FNetRefHandleManager::FReplicatedObjectData& ObjectData = NetRefHandleManager->GetReplicatedObjectDataNoCheck(InternalIndex);
-						if (ObjectData.IsSubObject())
+						if (ObjectData.IsOwnedSubObject())
 						{
 							FReplicationInfo& OwnerInfo = ReplicatedObjects[ObjectData.SubObjectRootIndex];
 							if (OwnerInfo.GetState() != EReplicatedObjectState::Invalid && OwnerInfo.GetState() < EReplicatedObjectState::PendingDestroy)
@@ -1384,7 +1384,7 @@ void FReplicationWriter::HandleDeliveredRecord(const FReplicationRecord::FRecord
 				if (Info.IsSubObject)
 				{
 					FNetRefHandleManager::FReplicatedObjectData& ObjectData = NetRefHandleManager->GetReplicatedObjectDataNoCheck(InternalIndex);
-					if (ObjectData.IsSubObject())
+					if (ObjectData.IsOwnedSubObject())
 					{
 						FReplicationInfo& OwnerInfo = ReplicatedObjects[ObjectData.SubObjectRootIndex];
 						MarkObjectDirty(ObjectData.SubObjectRootIndex, "HandleDeliveredRecordTearOff");
@@ -1605,7 +1605,7 @@ void FReplicationWriter::HandleDroppedRecord<FReplicationWriter::EReplicatedObje
 			//ensureMsgf(false, TEXT("Subobjects destroyed individually should not be canceled. Object: %s"), ToCStr(NetRefHandleManager->PrintObjectFromIndex(InternalIndex)));
 			// If the subobject owner still is replicated and valid
 			FNetRefHandleManager::FReplicatedObjectData& ObjectData = NetRefHandleManager->GetReplicatedObjectDataNoCheck(InternalIndex);
-			check(ObjectData.IsSubObject());
+			check(ObjectData.IsOwnedSubObject());
 
 			FNetBitArrayView ChangeMask(Info.GetChangeMaskStoragePointer(), Info.ChangeMaskBitCount);
 			Info.HasDirtyChangeMask |= ChangeMask.IsAnyBitSet();
@@ -1689,7 +1689,7 @@ void FReplicationWriter::HandleDroppedRecord<FReplicationWriter::EReplicatedObje
 		{
 			// If the subobject owner still is replicated and valid
 			FNetRefHandleManager::FReplicatedObjectData& ObjectData = NetRefHandleManager->GetReplicatedObjectDataNoCheck(InternalIndex);
-			check(ObjectData.IsSubObject());
+			check(ObjectData.IsOwnedSubObject());
 
 			// If owner is not pending destroy we mark it as dirty so that we can replicate subobject destruction properly
 			// We might get away with not doing this if owner or subobject does not have any unconfirmed changes in flight.
@@ -1985,7 +1985,7 @@ uint32 FReplicationWriter::WriteRootObjectsPendingDestroy(FNetSerializationConte
 		FReplicationInfo& Info = GetReplicationInfo(InternalIndex);
 		const FNetRefHandleManager::FReplicatedObjectData& ObjectData = NetRefHandleManager->GetReplicatedObjectDataNoCheck(InternalIndex);
 
-		if (ObjectData.IsSubObject())
+		if (ObjectData.IsOwnedSubObject())
 		{
 			// Have the root object query status of subobjects
 			continue;
diff --git a/Engine/Source/Runtime/Net/Iris/Public/Iris/ReplicationSystem/ObjectReplicationBridge.h b/Engine/Source/Runtime/Net/Iris/Public/Iris/ReplicationSystem/ObjectReplicationBridge.h
--- a/Engine/Source/Runtime/Net/Iris/Public/Iris/ReplicationSystem/ObjectReplicationBridge.h
+++ b/Engine/Source/Runtime/Net/Iris/Public/Iris/ReplicationSystem/ObjectReplicationBridge.h
@@ -148,7 +148,13 @@ public:
 	 */
 	IRISCORE_API void SetSubObjectNetCondition(FNetRefHandle SubObjectHandle, ELifetimeCondition Condition);
 
-	/** Get the handle of the root object of any replicated subobject. */
+	/** Returns the root object handle associated with any object. 
+	 * If the passed handle is a subobject it will return its root.  
+	 * If the passed handle is a rootobject it will return the same handle.
+	 */
+	IRISCORE_API FNetRefHandle GetRootObjectOfAnyObject(FNetRefHandle NetRefHandle) const;
+
+	/** Get the handle of the root object of any replicated subobject. SubObjectHandle MUST be binded to a subobject */
 	IRISCORE_API FNetRefHandle GetRootObjectOfSubObject(FNetRefHandle SubObjectHandle) const;
 
 	/** Add static destruction info, this is used when stably named objects are destroyed prior to starting replication */
