Iris
- Fixed bitstream corruption caused by edge case in hugeobject path where data was small enough to not be split but still did not fit in the current packet when including blob overhead.

diff --git a/Engine/Source/Runtime/Net/Iris/Private/Iris/ReplicationSystem/ReplicationWriter.cpp b/Engine/Source/Runtime/Net/Iris/Private/Iris/ReplicationSystem/ReplicationWriter.cpp
--- a/Engine/Source/Runtime/Net/Iris/Private/Iris/ReplicationSystem/ReplicationWriter.cpp
+++ b/Engine/Source/Runtime/Net/Iris/Private/Iris/ReplicationSystem/ReplicationWriter.cpp
@@ -3094,14 +3094,15 @@ FReplicationWriter::EWriteStatus FReplicationWriter::PrepareAndSendHugeObjectPay
 		HugeObjectBatchInfo.ParentInternalIndex = FNetRefHandleManager::InvalidInternalIndex;
 		const uint32 WriteHugeObjectFlags = WriteObjectFlag_Attachments | WriteObjectFlag_HugeObject;
 		const EWriteObjectStatus HugeObjectStatus = WriteObjectInBatch(Context, ObjectIndexForOOBAttachment, WriteHugeObjectFlags, HugeObjectBatchInfo);
-		if (!IsWriteObjectSuccess(HugeObjectStatus))
+		const uint32 BitsWrittenByHugeObjectAttachments = Writer.GetPosBits() - RollbackScope.GetStartPos();
+		if (!IsWriteObjectSuccess(HugeObjectStatus) || (BitsWrittenByHugeObjectAttachments == 0U))
 		{
 			// Need to call this in order to cleanup data associated with batch
 			HandleObjectBatchFailure(HugeObjectStatus, HugeObjectBatchInfo, WriteBitStreamInfo);
 
-			ensureMsgf(HugeObjectStatus == EWriteObjectStatus::BitStreamOverflow, TEXT("Expected split payload to not be able to generate other errors than overflow. Got %u"), unsigned(HugeObjectStatus));
 			// It's unexpected, but not a critical error, if no part of the payload could be sent.
 			// We do expect a smaller object to be sent though so that's why 0 is returned.
+			ensureMsgf(BitsWrittenByHugeObjectAttachments == 0U || HugeObjectStatus == EWriteObjectStatus::BitStreamOverflow, TEXT("Expected split payload to not be able to generate other errors than overflow. Got %u"), unsigned(HugeObjectStatus));
 
 			// Mark the context so that we can try to send the huge object in the next packet if we are allowed
 			WriteContext.bHasHugeObjectToSend = 1;
