Iris
* Deprecated NetObjectFactory::DestroyReplicatedObject  and replacted it with DetachedFromReplication
* The new function will receive callbacks for any object removed from replication not just those flagged to be destroyed.

#rnx

diff --git a/Engine/Plugins/Runtime/ReplicationSystemTestPlugin/Source/Private/Tests/ReplicationSystem/ReplicatedTestObjectFactory.cpp b/Engine/Plugins/Runtime/ReplicationSystemTestPlugin/Source/Private/Tests/ReplicationSystem/ReplicatedTestObjectFactory.cpp
--- a/Engine/Plugins/Runtime/ReplicationSystemTestPlugin/Source/Private/Tests/ReplicationSystem/ReplicatedTestObjectFactory.cpp
+++ b/Engine/Plugins/Runtime/ReplicationSystemTestPlugin/Source/Private/Tests/ReplicationSystem/ReplicatedTestObjectFactory.cpp
@@ -194,12 +194,9 @@ void UReplicatedTestObjectFactory::PostInit(const FPostInitContext& Context)
 	Instance->NetRefHandle = Context.Handle;
 }
 
-void UReplicatedTestObjectFactory::DestroyReplicatedObject(const FDestroyedContext& Context)
+void UReplicatedTestObjectFactory::DetachedFromReplication(const FDestroyedContext& Context)
 {
-	if (!Context.DestroyedInstance)
-	{
-		return;
-	}
+	check(Context.DestroyedInstance);
 
 	if (Context.DestroyReason == EReplicationBridgeDestroyInstanceReason::Destroy)
 	{
diff --git a/Engine/Plugins/Runtime/ReplicationSystemTestPlugin/Source/Private/Tests/ReplicationSystem/ReplicatedTestObjectFactory.h b/Engine/Plugins/Runtime/ReplicationSystemTestPlugin/Source/Private/Tests/ReplicationSystem/ReplicatedTestObjectFactory.h
--- a/Engine/Plugins/Runtime/ReplicationSystemTestPlugin/Source/Private/Tests/ReplicationSystem/ReplicatedTestObjectFactory.h
+++ b/Engine/Plugins/Runtime/ReplicationSystemTestPlugin/Source/Private/Tests/ReplicationSystem/ReplicatedTestObjectFactory.h
@@ -41,7 +41,7 @@ public:
 
 	// UNetObjectFactory interface
 	virtual void PostInit(const FPostInitContext& Context) override;
-	virtual void DestroyReplicatedObject(const FDestroyedContext& Context) override;
+	virtual void DetachedFromReplication(const FDestroyedContext& Context) override;
 	virtual TOptional<FWorldInfoData> GetWorldInfo(const FWorldInfoContext& Context) const override;
 	virtual void SubObjectCreatedFromReplication(UE::Net::FNetRefHandle RootObject, UE::Net::FNetRefHandle SubObjectCreated) override;
 	virtual void SubObjectDetachedFromReplication(const FDestroyedContext& Context) override;
diff --git a/Engine/Source/Runtime/Engine/Private/Net/Iris/ReplicationSystem/NetActorFactory.cpp b/Engine/Source/Runtime/Engine/Private/Net/Iris/ReplicationSystem/NetActorFactory.cpp
--- a/Engine/Source/Runtime/Engine/Private/Net/Iris/ReplicationSystem/NetActorFactory.cpp
+++ b/Engine/Source/Runtime/Engine/Private/Net/Iris/ReplicationSystem/NetActorFactory.cpp
@@ -489,11 +489,16 @@ void UNetActorFactory::SubObjectCreatedFromReplication(UE::Net::FNetRefHandle Ro
 	}
 }
 
-void UNetActorFactory::DestroyReplicatedObject(const FDestroyedContext& Context)
+void UNetActorFactory::DetachedFromReplication(const FDestroyedContext& Context)
 {
 	UEngineReplicationBridge* ActorBridge = CastChecked<UEngineReplicationBridge>(Bridge);
 	UNetDriver* NetDriver = ActorBridge->GetNetDriver();
 
+	if (Context.DestroyReason == EReplicationBridgeDestroyInstanceReason::DoNotDestroy)
+	{
+		return;
+	}
+
 	if (AActor* Actor = Cast<AActor>(Context.DestroyedInstance))
 	{
 		if ((Context.DestroyReason == EReplicationBridgeDestroyInstanceReason::TearOff) && !NetDriver->ShouldClientDestroyTearOffActors())
diff --git a/Engine/Source/Runtime/Engine/Private/Net/Iris/ReplicationSystem/NetSubObjectFactory.cpp b/Engine/Source/Runtime/Engine/Private/Net/Iris/ReplicationSystem/NetSubObjectFactory.cpp
--- a/Engine/Source/Runtime/Engine/Private/Net/Iris/ReplicationSystem/NetSubObjectFactory.cpp
+++ b/Engine/Source/Runtime/Engine/Private/Net/Iris/ReplicationSystem/NetSubObjectFactory.cpp
@@ -258,7 +258,7 @@ void UNetSubObjectFactory::SubObjectCreatedFromReplication(UE::Net::FNetRefHandl
 	ensureMsgf(false, TEXT("NetSubObjectFactory::SubObjectCreatedFromReplication should never be called since subobjects cannot have their own subobject list. RootObject: %s, SubObjectCreated: %s"), *Bridge->PrintObjectFromNetRefHandle(RootObject), *Bridge->PrintObjectFromNetRefHandle(SubObjectCreated));
 }
 
-void UNetSubObjectFactory::DestroyReplicatedObject(const FDestroyedContext& Context)
+void UNetSubObjectFactory::DetachedFromReplication(const FDestroyedContext& Context)
 {
 	// If the SubObject is being torn off it is up to owning actor to clean it up properly
 	if (Context.DestroyReason == EReplicationBridgeDestroyInstanceReason::TearOff)
@@ -266,11 +266,19 @@ void UNetSubObjectFactory::DestroyReplicatedObject(const FDestroyedContext& Cont
 		return;
 	}
 
+	if (Context.DestroyReason == EReplicationBridgeDestroyInstanceReason::DoNotDestroy)
+	{
+		// Nothing to do if the subobject shouldn't be destroyed
+		return;
+	}
+
 	if (!EnumHasAnyFlags(Context.DestroyFlags, EReplicationBridgeDestroyInstanceFlags::AllowDestroyInstanceFromRemote))
 	{
 		return;
 	}
 
+	// Find the subobject owner and trigger the callback. 
+    // TODO: This should be moved to NetFactorFactory::SubObjectDetachedFromReplication
 	if (AActor* ActorOwner = Cast<AActor>(Context.RootObject))
 	{
 		if (ensureMsgf(IsValid(ActorOwner) && !ActorOwner->IsUnreachable(), TEXT("UEngineReplicationBridge::DestroyInstanceFromRemote Destroyed subobject: %s has an invalid owner: %s"), *GetNameSafe(Context.DestroyedInstance), *GetPathNameSafe(Context.RootObject)))
diff --git a/Engine/Source/Runtime/Engine/Public/Net/Iris/ReplicationSystem/NetActorFactory.h b/Engine/Source/Runtime/Engine/Public/Net/Iris/ReplicationSystem/NetActorFactory.h
--- a/Engine/Source/Runtime/Engine/Public/Net/Iris/ReplicationSystem/NetActorFactory.h
+++ b/Engine/Source/Runtime/Engine/Public/Net/Iris/ReplicationSystem/NetActorFactory.h
@@ -188,7 +188,7 @@ public:
 
 	ENGINE_API virtual void SubObjectCreatedFromReplication(UE::Net::FNetRefHandle RootObject, UE::Net::FNetRefHandle SubObjectCreated) override;
 
-	ENGINE_API virtual void DestroyReplicatedObject(const FDestroyedContext& Context) override;
+	ENGINE_API virtual void DetachedFromReplication(const FDestroyedContext& Context) override;
 
 	ENGINE_API virtual TOptional<FWorldInfoData> GetWorldInfo(const FWorldInfoContext& Context) const override;
 
diff --git a/Engine/Source/Runtime/Engine/Public/Net/Iris/ReplicationSystem/NetSubObjectFactory.h b/Engine/Source/Runtime/Engine/Public/Net/Iris/ReplicationSystem/NetSubObjectFactory.h
--- a/Engine/Source/Runtime/Engine/Public/Net/Iris/ReplicationSystem/NetSubObjectFactory.h
+++ b/Engine/Source/Runtime/Engine/Public/Net/Iris/ReplicationSystem/NetSubObjectFactory.h
@@ -89,7 +89,7 @@ public:
 
 	ENGINE_API virtual void SubObjectCreatedFromReplication(UE::Net::FNetRefHandle RootObject, UE::Net::FNetRefHandle SubObjectCreated) override;
 
-	ENGINE_API virtual void DestroyReplicatedObject(const FDestroyedContext& Context) override;
+	ENGINE_API virtual void DetachedFromReplication(const FDestroyedContext& Context) override;
 
 	ENGINE_API virtual TOptional<FWorldInfoData> GetWorldInfo(const FWorldInfoContext& Context) const override;
 
diff --git a/Engine/Source/Runtime/Net/Iris/Private/Iris/ReplicationSystem/ObjectReplicationBridge.cpp b/Engine/Source/Runtime/Net/Iris/Private/Iris/ReplicationSystem/ObjectReplicationBridge.cpp
--- a/Engine/Source/Runtime/Net/Iris/Private/Iris/ReplicationSystem/ObjectReplicationBridge.cpp
+++ b/Engine/Source/Runtime/Net/Iris/Private/Iris/ReplicationSystem/ObjectReplicationBridge.cpp
@@ -1168,36 +1168,44 @@ void UObjectReplicationBridge::DetachInstanceFromRemote(FNetRefHandle Handle, ER
 
 	UE_LOG_OBJECTREPLICATIONBRIDGE(Verbose, TEXT("DetachInstanceFromRemote %s DestroyReason: %s DestroyFlags: %u"), *PrintObjectFromNetRefHandle(Handle), ToCStr(LexToString(DestroyReason)), unsigned(DestroyFlags));
 
-	// Tell factory about the imminent destruction
-	if (DestroyReason != EReplicationBridgeDestroyInstanceReason::DoNotDestroy)
+	UNetObjectFactory::FDestroyedContext DestroyContext
 	{
-		UNetObjectFactory::FDestroyedContext DestroyContext
-		{
-			.DestroyedInstance = GetObjectFromReferenceHandle(Handle),
-			.DestroyReason = DestroyReason,
-			.DestroyFlags = DestroyFlags,
-        };
+		.DestroyedInstance = GetObjectFromReferenceHandle(Handle),
+		.DestroyReason = DestroyReason,
+		.DestroyFlags = DestroyFlags,
+	};
 
-		if (DestroyContext.DestroyedInstance)
+	if (DestroyContext.DestroyedInstance)
+	{
+		// Check if it's a subobject and we need to tell it's RootObject
 		{
-			const FInternalNetRefIndex DestroyedObjectIndex = NetRefHandleManager->GetInternalIndex(Handle);
-			const FNetRefHandleManager::FReplicatedObjectData& DestroyedObjectData = NetRefHandleManager->GetReplicatedObjectData(DestroyedObjectIndex);
+			const FNetRefHandleManager::FReplicatedObjectData& DestroyedObjectData = NetRefHandleManager->GetReplicatedObjectData(NetRefHandleManager->GetInternalIndex(Handle));
 
-			// If this is a subobject, tell it's RootObject about the imminent destruction
-			if (DestroyedObjectData.SubObjectRootIndex != FNetRefHandleManager::InvalidInternalIndex)
+			const FInternalNetRefIndex RootObjectIndex = DestroyedObjectData.SubObjectRootIndex;
+
+			if (RootObjectIndex != FNetRefHandleManager::InvalidInternalIndex)
 			{
-				DestroyContext.RootObject = NetRefHandleManager->GetReplicatedObjectInstance(DestroyedObjectData.SubObjectRootIndex);
+				DestroyContext.RootObject = NetRefHandleManager->GetReplicatedObjectInstance(RootObjectIndex);
 
-				UNetObjectFactory* RootObjectFactory = GetNetFactory(NetRefHandleManager->GetReplicatedObjectData(DestroyedObjectData.SubObjectRootIndex).NetFactoryId);
+				UNetObjectFactory* RootObjectFactory = GetNetFactory(NetRefHandleManager->GetReplicatedObjectData(RootObjectIndex).NetFactoryId);
 				RootObjectFactory->SubObjectDetachedFromReplication(DestroyContext);
-			}			
+			}
+		}
 
-			// Now ask the object's factory to handle the destruction
-			if (ensureMsgf(NetFactoryId != InvalidNetObjectFactoryId, TEXT("Received invalid netfactory for %s: %s"), *Handle.ToString(), *GetPathNameSafe(DestroyContext.DestroyedInstance)))
+		// Now ask the object's factory to handle the destruction
+		if (ensureMsgf(NetFactoryId != InvalidNetObjectFactoryId, TEXT("Received invalid netfactory for %s: %s"), *Handle.ToString(), *GetPathNameSafe(DestroyContext.DestroyedInstance)))
+		{
+			UNetObjectFactory* Factory = GetNetFactory(NetFactoryId);
+
+PRAGMA_DISABLE_DEPRECATION_WARNINGS
+			// Call the deprecated function
+			if (DestroyReason != EReplicationBridgeDestroyInstanceReason::DoNotDestroy)
 			{
-				UNetObjectFactory* Factory = GetNetFactory(NetFactoryId);
 				Factory->DestroyReplicatedObject(DestroyContext);
 			}
+PRAGMA_ENABLE_DEPRECATION_WARNINGS
+
+			Factory->DetachedFromReplication(DestroyContext);
 		}
 	}
 	
diff --git a/Engine/Source/Runtime/Net/Iris/Public/Iris/ReplicationSystem/NetObjectFactory.h b/Engine/Source/Runtime/Net/Iris/Public/Iris/ReplicationSystem/NetObjectFactory.h
--- a/Engine/Source/Runtime/Net/Iris/Public/Iris/ReplicationSystem/NetObjectFactory.h
+++ b/Engine/Source/Runtime/Net/Iris/Public/Iris/ReplicationSystem/NetObjectFactory.h
@@ -140,11 +140,18 @@ public:
 	virtual void PostInit(const FPostInitContext& Context) {}
 
 	/**
-	 * Destroy a replicated object when it is detached from the remote authority.
-	 * A factory is not obligated to destroy the object for exemple if a TearOff is requested or the object is static and should only be hidden.
+	 * Callback triggered when a replicated object is no longer relevant on the client.
+	 * This is where a factory would destroy dynamic objects, reset stable objects or put objects back in a pool.
+	 * Note that this callback is called on remotes (clients) only.
 	 * @param Context Gives access to the instance that should be destroyed along with details on what type of destruction is requested.
+	 */
+	virtual void DetachedFromReplication(const FDestroyedContext& Context) PURE_VIRTUAL(UNetObjectFactory::DetachedFromReplication, );
+
+	/**
+	 * Replaced by DetachedFromReplication since it didn't receive callbacks for objects with EReplicationBridgeDestroyInstanceReason::DoNotDestroy reasons.
 	*/
-	virtual void DestroyReplicatedObject(const FDestroyedContext& Context) PURE_VIRTUAL(UNetObjectFactory::DestroyReplicatedObject, );
+	UE_DEPRECATED(5.7, "Replaced by DetachedFromReplication")
+	virtual void DestroyReplicatedObject(const FDestroyedContext& Context) {}
 
 	/** 
 	 * Optional callback triggered when a root object managed by this factory gets assigned a dynamic subobject. 
