Iris
* Remove IntCastChecked when reading factory ID and handle null factories instead

#rnx

diff --git a/Engine/Source/Runtime/Net/Iris/Private/Iris/ReplicationSystem/ObjectReplicationBridge.cpp b/Engine/Source/Runtime/Net/Iris/Private/Iris/ReplicationSystem/ObjectReplicationBridge.cpp
--- a/Engine/Source/Runtime/Net/Iris/Private/Iris/ReplicationSystem/ObjectReplicationBridge.cpp
+++ b/Engine/Source/Runtime/Net/Iris/Private/Iris/ReplicationSystem/ObjectReplicationBridge.cpp
@@ -491,7 +491,7 @@ void UObjectReplicationBridge::ReadAndExecuteDestructionInfoFromRemote(FReplicat
 	FNetObjectReference ReferenceToDestroy;
 	ReadFullNetObjectReference(Context.SerializationContext, ReferenceToDestroy);
 
-	const FNetObjectFactoryId NetFactoryId = IntCastChecked<FNetObjectFactoryId>(Context.SerializationContext.GetBitStreamReader()->ReadBits(FNetObjectFactoryRegistry::GetMaxBits()));
+	const FNetObjectFactoryId NetFactoryId = (FNetObjectFactoryId)Context.SerializationContext.GetBitStreamReader()->ReadBits(FNetObjectFactoryRegistry::GetMaxBits());
 
 	// Destroy the reference
 	// Resolve the reference in order to be able to destroy it
@@ -1103,7 +1103,7 @@ bool UObjectReplicationBridge::WriteNetRefHandleCreationInfo(FReplicationBridgeS
 	const FNetObjectFactoryId FactoryId = ObjectData.NetFactoryId;
 
 	UNetObjectFactory* Factory = GetNetFactory(FactoryId);
-
+	check(Factory);
 	TUniquePtr<FNetObjectCreationHeader> Header = Factory->CreateHeader(Handle, Protocol->ProtocolIdentifier);
 
 	if (!Header.IsValid())
@@ -1139,7 +1139,7 @@ bool UObjectReplicationBridge::CacheNetRefHandleCreationInfo(FNetRefHandle Handl
 	const FNetObjectFactoryId FactoryId = NetRefHandleManager->GetReplicatedObjectDataNoCheck(ObjectIndex).NetFactoryId;
 
 	UNetObjectFactory* Factory = GetNetFactory(FactoryId);
-
+	check(Factory);
 	TUniquePtr<FNetObjectCreationHeader> Header = Factory->CreateHeader(Handle, Protocol->ProtocolIdentifier);
 
 	if (!Header.IsValid())
@@ -1187,25 +1187,28 @@ void UObjectReplicationBridge::DetachInstanceFromRemote(FNetRefHandle Handle, ER
 			{
 				DestroyContext.RootObject = NetRefHandleManager->GetReplicatedObjectInstance(RootObjectIndex);
 
-				UNetObjectFactory* RootObjectFactory = GetNetFactory(NetRefHandleManager->GetReplicatedObjectData(RootObjectIndex).NetFactoryId);
-				RootObjectFactory->SubObjectDetachedFromReplication(DestroyContext);
+				if (UNetObjectFactory* RootObjectFactory = GetNetFactory(NetRefHandleManager->GetReplicatedObjectData(RootObjectIndex).NetFactoryId))
+				{
+					RootObjectFactory->SubObjectDetachedFromReplication(DestroyContext);
+				}
 			}
 		}
 
 		// Now ask the object's factory to handle the destruction
 		if (ensureMsgf(NetFactoryId != InvalidNetObjectFactoryId, TEXT("Received invalid netfactory for %s: %s"), *Handle.ToString(), *GetPathNameSafe(DestroyContext.DestroyedInstance)))
 		{
-			UNetObjectFactory* Factory = GetNetFactory(NetFactoryId);
-
-PRAGMA_DISABLE_DEPRECATION_WARNINGS
-			// Call the deprecated function
-			if (DestroyReason != EReplicationBridgeDestroyInstanceReason::DoNotDestroy)
+			if (UNetObjectFactory* Factory = GetNetFactory(NetFactoryId))
 			{
-				Factory->DestroyReplicatedObject(DestroyContext);
-			}
+PRAGMA_DISABLE_DEPRECATION_WARNINGS
+				// Call the deprecated function
+				if (DestroyReason != EReplicationBridgeDestroyInstanceReason::DoNotDestroy)
+				{
+					Factory->DestroyReplicatedObject(DestroyContext);
+				}
 PRAGMA_ENABLE_DEPRECATION_WARNINGS
 
-			Factory->DetachedFromReplication(DestroyContext);
+				Factory->DetachedFromReplication(DestroyContext);
+			}
 		}
 	}
 	
@@ -1290,7 +1293,7 @@ FReplicationBridgeCreateNetRefHandleResult UObjectReplicationBridge::CreateNetRe
 
 	FNetBitStreamReader* Reader = Context.SerializationContext.GetBitStreamReader();
 
-	const FNetObjectFactoryId FactoryId = IntCastChecked<FNetObjectFactoryId>(Reader->ReadBits(FNetObjectFactoryRegistry::GetMaxBits()));
+	const FNetObjectFactoryId FactoryId = (FNetObjectFactoryId)Reader->ReadBits(FNetObjectFactoryRegistry::GetMaxBits());
 	UNetObjectFactory* Factory = GetNetFactory(FactoryId);
 	if (Factory == nullptr)
 	{
@@ -1494,7 +1497,7 @@ void UObjectReplicationBridge::SubObjectCreatedFromReplication(UE::Net::Private:
 
 	const FNetObjectFactoryId FactoryId = NetRefHandleManager->GetReplicatedObjectDataNoCheck(RootObjectIndex).NetFactoryId;
 	UNetObjectFactory* Factory = GetNetFactory(FactoryId);
-
+	check(Factory);
 	Factory->SubObjectCreatedFromReplication(NetRefHandleManager->GetNetRefHandleFromInternalIndex(RootObjectIndex), SubObjectCreated);
 }
 
@@ -1505,7 +1508,7 @@ void UObjectReplicationBridge::PostApplyInitialState(UE::Net::Private::FInternal
 
 	const FNetObjectFactoryId FactoryId = NetRefHandleManager->GetReplicatedObjectDataNoCheck(InternalObjectIndex).NetFactoryId;
 	UNetObjectFactory* Factory = GetNetFactory(FactoryId);
-
+	check(Factory);
 	const UNetObjectFactory::FPostInitContext Context 
 	{ 
 		.Instance = NetRefHandleManager->GetReplicatedObjectInstance(InternalObjectIndex),
@@ -2863,7 +2866,7 @@ float UObjectReplicationBridge::GetPollFrequencyOfRootObject(UE::Net::Private::F
 	UObject* RootObjectInstance = NetRefHandleManager->GetReplicatedObjectInstance(RootObjectIndex);
 	FNetRefHandle RootObjectHandle = NetRefHandleManager->GetNetRefHandleFromInternalIndex(RootObjectIndex);
 	
-	FNetObjectFactoryId FactoryId = NetRefHandleManager->GetReplicatedObjectDataNoCheck(RootObjectIndex).NetFactoryId;
+	const FNetObjectFactoryId FactoryId = NetRefHandleManager->GetReplicatedObjectDataNoCheck(RootObjectIndex).NetFactoryId;
 	if (FactoryId == UE::Net::InvalidNetObjectFactoryId)
 	{
 		// TODO: Investigate how scopeable objects are Invalid
diff --git a/Engine/Source/Runtime/Net/Iris/Public/Iris/ReplicationSystem/NetObjectFactoryRegistry.h b/Engine/Source/Runtime/Net/Iris/Public/Iris/ReplicationSystem/NetObjectFactoryRegistry.h
--- a/Engine/Source/Runtime/Net/Iris/Public/Iris/ReplicationSystem/NetObjectFactoryRegistry.h
+++ b/Engine/Source/Runtime/Net/Iris/Public/Iris/ReplicationSystem/NetObjectFactoryRegistry.h
@@ -9,7 +9,7 @@
 
 // Define this to increase the maximum amount of NetObjectFactories that can exist
 #ifndef UE_IRIS_MAX_NETOBJECT_FACTORIES
-	#define UE_IRIS_MAX_NETOBJECT_FACTORIES 8
+	#define UE_IRIS_MAX_NETOBJECT_FACTORIES 16
 #endif
 
 namespace UE::Net
