Make some Iris filtering behavior more consistent.
- Objects that have a spatial filter set at runtime will now have their locations properly updated. Call OptionallySetObjectRequiresFrequentWorldLocationUpdate from UReplicationSystem::SetFilter, in case the new filter is spatial and frequent updates weren't enabled when the object started replicating.
- Actors and subclasses without explicit FilterConfig overrides will now be filtered based on their relevancy flags like bAlwaysRelevant instead of being always relevant by default. Removed the FilterConfigs for Actor, PlayerState, and Pawn from BaseEngine.ini.

diff --git a/Engine/Config/BaseEngine.ini b/Engine/Config/BaseEngine.ini
--- a/Engine/Config/BaseEngine.ini
+++ b/Engine/Config/BaseEngine.ini
@@ -1454,14 +1454,11 @@ RequiredNetDriverChannelClassName=/Script/Engine.DataStreamChannel
 DefaultSpatialFilterName=Spatial
 ; Filter configs should typically be configured per game
 ; Set the class-based dynamic filters configs here. Using None means the objects of this class are always relevant.
+; If there is no FilterConfig for a class, and the class should be spatially filtered, it will use the filter specified by UObjectReplicationBridgeConfig::DefaultSpatialFilterName.
 !FilterConfigs=ClearArray
 ;+FilterConfigs=(ClassName=/Script/Engine.LevelScriptActor, DynamicFilterName=NotRouted) ; NotRouted means that this type will not be replicated
-+FilterConfigs=(ClassName=/Script/Engine.Actor, DynamicFilterName=None)) ; By default classes derived from AActor are always relevant (because not filtered dynamically)
 ; Info types aren't supposed to have physical representation and rely on a static priority
 +FilterConfigs=(ClassName=/Script/Engine.Info, DynamicFilterName=None)
-+FilterConfigs=(ClassName=/Script/Engine.PlayerState, DynamicFilterName=None)
-; Pawns can be spatially filtered...
-+FilterConfigs=(ClassName=/Script/Engine.Pawn, DynamicFilterName=Spatial))
 
 ; PrioritizerConfigs
 ; PlayerState is by default bAlwaysRelevant but the default player state prioritizer handles this.
diff --git a/Engine/Source/Runtime/Net/Iris/Private/Iris/ReplicationSystem/ReplicationSystem.cpp b/Engine/Source/Runtime/Net/Iris/Private/Iris/ReplicationSystem/ReplicationSystem.cpp
--- a/Engine/Source/Runtime/Net/Iris/Private/Iris/ReplicationSystem/ReplicationSystem.cpp
+++ b/Engine/Source/Runtime/Net/Iris/Private/Iris/ReplicationSystem/ReplicationSystem.cpp
@@ -54,6 +54,13 @@ static FAutoConsoleVariableRef CVarAllowAttachmentSendPolicyFlags(TEXT("net.Iris
 static bool bOnlyResetDirtinessForQuantizedObjects = true;
 static FAutoConsoleVariableRef CVarOnlyResetDirtinessForQuantizedObjects(TEXT("net.Iris.OnlyResetDirtinessForQuantizedObjects"), bOnlyResetDirtinessForQuantizedObjects, TEXT("Only Reset Dirtiness For QuantizedObjects, optimization that only resets dirtiness for objects actually considered dirty."));
 
+static bool bSetFilterSetsRequiresFrequentLocationUpdates = true;
+static FAutoConsoleVariableRef CVarSetFilterSetsRequiresFrequentLocationUpdates(
+	TEXT("net.Iris.SetFilterSetsRequiresFrequentLocationUpdates"),
+	bSetFilterSetsRequiresFrequentLocationUpdates,
+	TEXT("Call OptionallySetObjectRequiresFrequentWorldLocationUpdate from SetFilter in case the new filter is spatial and the object needs frequent updates. ")
+	TEXT("This being enabled is the intended behavior, only disable in case of an unintended regression."));
+
 }
 
 namespace UE::Net::Private
@@ -1491,7 +1498,17 @@ bool UReplicationSystem::SetFilter(FNetRefHandle Handle, UE::Net::FNetObjectFilt
 	}
 
 	FReplicationFiltering& Filtering = Impl->ReplicationSystemInternal.GetFiltering();
-	return Filtering.SetFilter(ObjectInternalIndex, Filter, FilterConfigProfile);
+	const bool SetFilterResult = Filtering.SetFilter(ObjectInternalIndex, Filter, FilterConfigProfile);
+
+	if (ReplicationSystemCVars::bSetFilterSetsRequiresFrequentLocationUpdates)
+	{
+		// If the new filter is spatial, the object might now require world location updates.
+		UObjectReplicationBridge* Bridge = Impl->ReplicationSystemInternal.GetReplicationBridge();
+		const bool bWantsToBeDormant = Bridge->GetObjectWantsToBeDormant(Handle);
+		Bridge->OptionallySetObjectRequiresFrequentWorldLocationUpdate(Handle, !bWantsToBeDormant);
+	}
+
+	return SetFilterResult;
 }
 
 UE::Net::FNetObjectFilterHandle UReplicationSystem::GetFilterHandle(const FName FilterName) const
