Iris
* Made UObjectReplicationBridge::GetNetFactory public

#rnx

diff --git a/Engine/Source/Runtime/Net/Iris/Private/Iris/ReplicationSystem/ObjectReplicationBridge.cpp b/Engine/Source/Runtime/Net/Iris/Private/Iris/ReplicationSystem/ObjectReplicationBridge.cpp
--- a/Engine/Source/Runtime/Net/Iris/Private/Iris/ReplicationSystem/ObjectReplicationBridge.cpp
+++ b/Engine/Source/Runtime/Net/Iris/Private/Iris/ReplicationSystem/ObjectReplicationBridge.cpp
@@ -1077,9 +1077,13 @@ bool UObjectReplicationBridge::WriteNetRefHandleCreationInfo(FReplicationBridgeS
 	// If there is a cached header for this handle, use it!
 	if (TUniquePtr<const FNetObjectCreationHeader>* CachedHeader = CachedCreationHeaders.Find(Handle))
 	{
+		//TODO: Right now cached creation data is only for destroyed or torn-off objects so the data can be considered const
+		//      If we eventually use it to store actively replicated objects, some creation data will want to be updated
+		//	    before being sent.  We should add a trait and a NetObjectFactory::UpdateCachedData to do so.
+
 		const FNetObjectCreationHeader* Header = CachedHeader->Get();
 		UNetObjectFactory* Factory = GetNetFactory(Header->GetNetFactoryId());
-		UE_LOG_BRIDGEID(LogIris, VeryVerbose, TEXT("WriteNetRefHandleCreationInfo (cached) for %s : Factory: %s (id:%u) Header: %s"), *Handle.ToString(), *GetNameSafe(Factory), Header->GetNetFactoryId(), *Header->ToString());
+		UE_LOG_BRIDGEID(LogIris, Verbose, TEXT("WriteNetRefHandleCreationInfo (cached) for %s : Factory: %s (id:%u) Header: %s"), *Handle.ToString(), *GetNameSafe(Factory), Header->GetNetFactoryId(), *Header->ToString());
 		return Factory->WriteHeader(Handle, Context.SerializationContext, Header);
 	}
 	
diff --git a/Engine/Source/Runtime/Net/Iris/Public/Iris/ReplicationSystem/ObjectReplicationBridge.h b/Engine/Source/Runtime/Net/Iris/Public/Iris/ReplicationSystem/ObjectReplicationBridge.h
--- a/Engine/Source/Runtime/Net/Iris/Public/Iris/ReplicationSystem/ObjectReplicationBridge.h
+++ b/Engine/Source/Runtime/Net/Iris/Public/Iris/ReplicationSystem/ObjectReplicationBridge.h
@@ -157,6 +157,9 @@ public:
 	/** Get the handle of the root object of any replicated subobject. SubObjectHandle MUST be binded to a subobject */
 	IRISCORE_API FNetRefHandle GetRootObjectOfSubObject(FNetRefHandle SubObjectHandle) const;
 
+	/** Returns the associated factory created for this bridge. */
+	IRISCORE_API UNetObjectFactory* GetNetFactory(UE::Net::FNetObjectFactoryId FactoryId) const;
+
 	/** Add static destruction info, this is used when stably named objects are destroyed prior to starting replication */
 	IRISCORE_API void AddStaticDestructionInfo(const FString& ObjectPath, const UObject* Outer, const UE::Net::FDestructionParameters& Parameters);
 
@@ -357,8 +360,6 @@ protected:
 	/** Lookup the UObject associated with the provided Handle. This function will not try to resolve the reference. */
 	IRISCORE_API UObject* GetObjectFromReferenceHandle(FNetRefHandle RefHandle) const;
 
-	IRISCORE_API UNetObjectFactory* GetNetFactory(UE::Net::FNetObjectFactoryId FactoryId) const;
-
 	/** Helper method that calls provided PreUpdateFunction and polls state data for all replicated instances with the NeedsPoll trait. */
 	using FInstancePreUpdateFunction = TFunction<void(TArrayView<UObject*>, const UObjectReplicationBridge*)>;
 
