Set ObjectRecord.NextIndex to invalid index on creation to mitigate infinite loop in FReplicationWriter::HasInFlightStateChanges

diff --git a/Engine/Source/Runtime/Net/Iris/Private/Iris/ReplicationSystem/ReplicationWriter.cpp b/Engine/Source/Runtime/Net/Iris/Private/Iris/ReplicationSystem/ReplicationWriter.cpp
--- a/Engine/Source/Runtime/Net/Iris/Private/Iris/ReplicationSystem/ReplicationWriter.cpp
+++ b/Engine/Source/Runtime/Net/Iris/Private/Iris/ReplicationSystem/ReplicationWriter.cpp
@@ -1843,6 +1843,9 @@ void FReplicationWriter::CreateObjectRecord(const FNetBitArrayView* ChangeMask,
 	FReplicationRecord::FRecordInfo& RecordInfo = OutRecord.Record;
 
 	RecordInfo.Index = ObjectInfo.InternalIndex;
+	
+	RecordInfo.NextIndex = FReplicationRecord::InvalidReplicationRecordIndex;
+
 	RecordInfo.ReplicatedObjectState = ObjectInfo.AttachmentType == ENetObjectAttachmentType::HugeObject ? uint8(EReplicatedObjectState::HugeObject) : (uint8)Info.GetState();
 	RecordInfo.HasChangeMask = ChangeMask ? 1U : 0U;
 	RecordInfo.HasAttachments = (OutRecord.AttachmentRecord.IsValid() ? 1U : 0U);
