Iris - Fixed fallout from multiserver proxy support.

#rnx

diff --git a/Engine/Plugins/Runtime/ReplicationSystemTestPlugin/Source/Private/Tests/ReplicationSystem/TestReplicationOperations.cpp b/Engine/Plugins/Runtime/ReplicationSystemTestPlugin/Source/Private/Tests/ReplicationSystem/TestReplicationOperations.cpp
--- a/Engine/Plugins/Runtime/ReplicationSystemTestPlugin/Source/Private/Tests/ReplicationSystem/TestReplicationOperations.cpp
+++ b/Engine/Plugins/Runtime/ReplicationSystemTestPlugin/Source/Private/Tests/ReplicationSystem/TestReplicationOperations.cpp
@@ -420,6 +420,8 @@ UE_NET_TEST_FIXTURE(FTestReplicationOperationsFixture, InitPropertyIsNotUpdatedW
 
 	// Copy data from test object to buffer
 	FNetSerializationContext NetContext;
+	FInternalNetSerializationContext InternalContext;
+	NetContext.SetInternalContext(&InternalContext);
 	NetContext.SetIsInitState(false);
 	FReplicationInstanceOperations::Quantize(NetContext, AlignedBuffer, GetChangeMaskWriter(Protocol->ChangeMaskBitCount), ReplicationBridge->GetReplicationInstanceProtocol(TestObject->NetRefHandle), Protocol);
 
diff --git a/Engine/Source/Runtime/Net/Iris/Private/Iris/ReplicationSystem/DequantizeAndApplyHelper.cpp b/Engine/Source/Runtime/Net/Iris/Private/Iris/ReplicationSystem/DequantizeAndApplyHelper.cpp
--- a/Engine/Source/Runtime/Net/Iris/Private/Iris/ReplicationSystem/DequantizeAndApplyHelper.cpp
+++ b/Engine/Source/Runtime/Net/Iris/Private/Iris/ReplicationSystem/DequantizeAndApplyHelper.cpp
@@ -256,7 +256,7 @@ void FDequantizeAndApplyHelper::CallLegacyPostApplyFunctions(FContext* Context,
 	CallLegacyFunctionForEachOwner(*Context, [](UObject* Object) { Object->PostNetReceive(); });
 
 	// CallRepNotifies
-	if (NetSerializationContext.GetInternalContext()->ReplicationSystem->AllowExecuteRepNotify())
+	if (NetSerializationContext.GetInternalContext()->bAllowExecuteRepNotify)
 	{
 		if (Context->bNeedsLegacyCallbacks)
 		{
diff --git a/Engine/Source/Runtime/Net/Iris/Private/Iris/Serialization/InternalNetSerializationContext.cpp b/Engine/Source/Runtime/Net/Iris/Private/Iris/Serialization/InternalNetSerializationContext.cpp
--- a/Engine/Source/Runtime/Net/Iris/Private/Iris/Serialization/InternalNetSerializationContext.cpp
+++ b/Engine/Source/Runtime/Net/Iris/Private/Iris/Serialization/InternalNetSerializationContext.cpp
@@ -15,7 +15,8 @@ FInternalNetSerializationContext::FInternalNetSerializationContext(UReplicationS
 , bDowngradeAutonomousProxyRole(0)
 , bInlineObjectReferenceExports(0)
 , bSerializeObjectReferencesAsRemoteIds(InReplicationSystem->IsUsingRemoteObjectReferences() ? 1 : 0)
-, bAlwaysSwapRolesOnReplication(0)
+, bAlwaysSwapRolesOnReplication(InReplicationSystem->AlwaysSwapRolesOnReplication() ? 1 : 0)
+, bAllowExecuteRepNotify(ReplicationSystem->AllowExecuteRepNotify() ? 1 : 0)
 {
 }
 
@@ -36,8 +37,7 @@ void* FInternalNetSerializationContext::Realloc(void* PrevAddress, SIZE_T NewSiz
 
 void FInternalNetSerializationContext::Init(const FInitParameters& InitParams)
 {
-	ReplicationSystem = InitParams.ReplicationSystem;
-	ObjectReferenceCache = &InitParams.ReplicationSystem->GetReplicationSystemInternal()->GetObjectReferenceCache();
+	*this = FInternalNetSerializationContext(InitParams.ReplicationSystem);
 	PackageMap = InitParams.PackageMap;
 	ResolveContext = InitParams.ObjectResolveContext;
 }
diff --git a/Engine/Source/Runtime/Net/Iris/Private/Iris/Serialization/InternalNetSerializationContext.h b/Engine/Source/Runtime/Net/Iris/Private/Iris/Serialization/InternalNetSerializationContext.h
--- a/Engine/Source/Runtime/Net/Iris/Private/Iris/Serialization/InternalNetSerializationContext.h
+++ b/Engine/Source/Runtime/Net/Iris/Private/Iris/Serialization/InternalNetSerializationContext.h
@@ -58,6 +58,9 @@ public:
 
 	// If true, swap local and remote roles when replicating.
 	uint32 bAlwaysSwapRolesOnReplication : 1;
+
+	// If true, allow calling rep notifies
+	uint32 bAllowExecuteRepNotify : 1;
 };
 
 inline FInternalNetSerializationContext::FInternalNetSerializationContext()
@@ -67,6 +70,8 @@ inline FInternalNetSerializationContext::FInternalNetSerializationContext()
 , bDowngradeAutonomousProxyRole(0)
 , bInlineObjectReferenceExports(0)
 , bSerializeObjectReferencesAsRemoteIds(0)
+, bAlwaysSwapRolesOnReplication(0)
+, bAllowExecuteRepNotify(1)
 {}
 
 // Scope used when we actually want to export references.
