Allow the use of UNetDriver::ReplicateAutonomousAsSimulated() in Iris and fix a bug in Iris when using UNetDriver::IsRoleSwappingOnAllActorsEnabled().

diff --git a/Engine/Plugins/Runtime/MultiServerReplication/Source/MultiServerReplication/Private/MultiServerProxy.cpp b/Engine/Plugins/Runtime/MultiServerReplication/Source/MultiServerReplication/Private/MultiServerProxy.cpp
--- a/Engine/Plugins/Runtime/MultiServerReplication/Source/MultiServerReplication/Private/MultiServerProxy.cpp
+++ b/Engine/Plugins/Runtime/MultiServerReplication/Source/MultiServerReplication/Private/MultiServerProxy.cpp
@@ -339,7 +339,7 @@ void UProxyListenerNotify::ConnectToGameServer(UNetConnection* ProxyConnection,
 		URL.AddOption(*FString::Printf(TEXT("HandshakeId=%u"), ParentClientHandshakeId));
 
 		// Maybe this should be a control message since it changes the server setting (it's global driver setting).
-		URL.AddOption(TEXT("AutonomousAsSimulated"));
+		URL.AddOption(TEXT("Proxy"));
 
 		// Start the connection flow to the game server.
 		GameServerConnectionState->GameServerNotify = NewObject<UGameServerNotify>();
diff --git a/Engine/Source/Runtime/Engine/Classes/Engine/ChildConnection.h b/Engine/Source/Runtime/Engine/Classes/Engine/ChildConnection.h
--- a/Engine/Source/Runtime/Engine/Classes/Engine/ChildConnection.h
+++ b/Engine/Source/Runtime/Engine/Classes/Engine/ChildConnection.h
@@ -105,6 +105,9 @@ public:
 
 	ENGINE_API virtual void InitChildConnection(UNetDriver* InDriver, UNetConnection* InParent);
 
+	ENGINE_API virtual void SetProxyConnection(bool bEnabled) override;
+	ENGINE_API virtual bool IsProxyConnection() const override;
+
 private:
 	void AssignConnectionHandle();
 };
diff --git a/Engine/Source/Runtime/Engine/Classes/Engine/NetConnection.h b/Engine/Source/Runtime/Engine/Classes/Engine/NetConnection.h
--- a/Engine/Source/Runtime/Engine/Classes/Engine/NetConnection.h
+++ b/Engine/Source/Runtime/Engine/Classes/Engine/NetConnection.h
@@ -1457,9 +1457,15 @@ public:
 	/** Sets whether we handle opening channels with an index that already exists, used by replays to fast forward the packet stream */
 	void SetAllowExistingChannelIndex(bool bAllow);
 
+	/** Set to true if this connection's remote endpoint is a proxy. */
+	ENGINE_API virtual void SetProxyConnection(bool bEnabled);
+	ENGINE_API virtual bool IsProxyConnection() const;
+
 private:
 	void RestoreRemappedChannel(const int32 ChIndex);
 
+	bool bIsProxyConnection = false;
+
 public:
 
 	/**
diff --git a/Engine/Source/Runtime/Engine/Classes/Engine/NetDriver.h b/Engine/Source/Runtime/Engine/Classes/Engine/NetDriver.h
--- a/Engine/Source/Runtime/Engine/Classes/Engine/NetDriver.h
+++ b/Engine/Source/Runtime/Engine/Classes/Engine/NetDriver.h
@@ -1492,10 +1492,6 @@ public:
 	ENGINE_API void EnableReplicationOfActorsWithAnyRole(bool bEnable);
 	ENGINE_API bool IsReplicationOfActorsWithAnyRoleEnabled() const;
 
-	/** (Generic Replicaton Only) Treat actors with ROLE_AutonomousProxy as ROLE_SimulatedProxy when replicating state (applies to all client connections). */
-	ENGINE_API void ReplicateAutonomousAsSimulated(bool bEnabled);
-	ENGINE_API bool IsReplicatingAutonomousAsSimulated() const;
-
 	/** (Generic Replicaton Only) Return true if an actor's role can be downgraded during replication to a given connection.  */
 	virtual bool CanDowngradeActorRole(UNetConnection* Connection, AActor* Actor) const { return true; };
 
@@ -1505,7 +1501,6 @@ private:
 	bool bCallPreReplication = true;
 	bool bSwapRolesOfAllActors = false;
 	bool bReplicationOfActorsWithAnyRole = false;
-	bool bReplicatingAutonomousAsSimulated = false;
 
 	/** A set of replicators that have queued RPCs for transactions (server only) */
 	TSet< FObjectReplicator* > ReplicatorsWithPendingTransactions;
diff --git a/Engine/Source/Runtime/Engine/Private/DataChannel.cpp b/Engine/Source/Runtime/Engine/Private/DataChannel.cpp
--- a/Engine/Source/Runtime/Engine/Private/DataChannel.cpp
+++ b/Engine/Source/Runtime/Engine/Private/DataChannel.cpp
@@ -3811,7 +3811,7 @@ int64 UActorChannel::ReplicateActor()
 
 	RepFlags.bNetSimulated	= (Actor->GetRemoteRole() == ROLE_SimulatedProxy);
 
-	if (Actor->GetRemoteRole() == ROLE_AutonomousProxy && Connection->GetDriver()->IsReplicatingAutonomousAsSimulated())
+	if (Actor->GetRemoteRole() == ROLE_AutonomousProxy && Connection->IsProxyConnection())
 	{
 		RepFlags.bNetSimulated = true;
 	}
diff --git a/Engine/Source/Runtime/Engine/Private/NetConnection.cpp b/Engine/Source/Runtime/Engine/Private/NetConnection.cpp
--- a/Engine/Source/Runtime/Engine/Private/NetConnection.cpp
+++ b/Engine/Source/Runtime/Engine/Private/NetConnection.cpp
@@ -1580,6 +1580,16 @@ void UChildConnection::InitChildConnection(UNetDriver* InDriver, UNetConnection*
 	AssignConnectionHandle();
 }
 
+void UChildConnection::SetProxyConnection(bool bEnabled)
+{
+	UE_LOG(LogNet, Error, TEXT("Cannot set the child connection %s to be a proxy connection."), *GetName());
+}
+
+bool UChildConnection::IsProxyConnection() const
+{
+	return Parent->IsProxyConnection();
+}
+
 void UChildConnection::AssignConnectionHandle()
 {
 	if (const UNetConnection* ParentConn = Parent.Get())
@@ -4212,6 +4222,22 @@ void UNetConnection::SetAllowExistingChannelIndex(bool bAllow)
 	ChannelIndexMap.Reset();
 }
 
+void UNetConnection::SetProxyConnection(bool bEnabled)
+{
+	bIsProxyConnection = bEnabled;
+
+	if (UReplicationSystem* DriverReplicationSystem = Driver->GetReplicationSystem())
+	{
+		const uint32 ConnectionId = GetConnectionHandle().GetParentConnectionId();
+		DriverReplicationSystem->SetProxyConnection(ConnectionId, bEnabled);
+	}
+}
+
+bool UNetConnection::IsProxyConnection() const
+{
+	return bIsProxyConnection;
+}
+
 void UNetConnection::SetIgnoreActorBunches(bool bInIgnoreActorBunches, TSet<FNetworkGUID>&& InIgnoredBunchGuids)
 {
 	check(IsInternalAck());
diff --git a/Engine/Source/Runtime/Engine/Private/NetDriver.cpp b/Engine/Source/Runtime/Engine/Private/NetDriver.cpp
--- a/Engine/Source/Runtime/Engine/Private/NetDriver.cpp
+++ b/Engine/Source/Runtime/Engine/Private/NetDriver.cpp
@@ -2320,16 +2320,6 @@ bool UNetDriver::IsReplicationOfActorsWithAnyRoleEnabled() const
 	return bReplicationOfActorsWithAnyRole;
 }
 
-void UNetDriver::ReplicateAutonomousAsSimulated(bool bEnabled)
-{
-	bReplicatingAutonomousAsSimulated = bEnabled;
-}
-
-bool UNetDriver::IsReplicatingAutonomousAsSimulated() const
-{
-	return bReplicatingAutonomousAsSimulated;
-}
-
 ENetMode UNetDriver::GetNetMode() const
 {
 	// Special case for PIE and game worlds - forcing dedicated server behavior if requested
diff --git a/Engine/Source/Runtime/Engine/Private/World.cpp b/Engine/Source/Runtime/Engine/Private/World.cpp
--- a/Engine/Source/Runtime/Engine/Private/World.cpp
+++ b/Engine/Source/Runtime/Engine/Private/World.cpp
@@ -7369,9 +7369,11 @@ void UWorld::NotifyControlMessage(UNetConnection* Connection, uint8 MessageType,
 						break;
 					}
 
-					if (InURL.HasOption(TEXT("AutonomousAsSimulated")))
+					if (InURL.HasOption(TEXT("Proxy")))
 					{
-						NetDriver->ReplicateAutonomousAsSimulated(true);
+						// Only the parent connection needs to be set to be a proxy connection. Child connections will
+						// use the state of the parent connection when determining if it's a proxy connection.
+						Connection->SetProxyConnection(true);
 					}
 
 					int32 SplitscreenCount = FMath::Min(Connection->Children.Num() + 1, 255);
diff --git a/Engine/Source/Runtime/Net/Iris/Private/Iris/ReplicationSystem/Conditionals/ReplicationConditionals.cpp b/Engine/Source/Runtime/Net/Iris/Private/Iris/ReplicationSystem/Conditionals/ReplicationConditionals.cpp
--- a/Engine/Source/Runtime/Net/Iris/Private/Iris/ReplicationSystem/Conditionals/ReplicationConditionals.cpp
+++ b/Engine/Source/Runtime/Net/Iris/Private/Iris/ReplicationSystem/Conditionals/ReplicationConditionals.cpp
@@ -189,6 +189,19 @@ void FReplicationConditionals::RemoveConnection(uint32 ConnectionId)
 	ConnectionInfo.ObjectConditionals.Empty();
 }
 
+void FReplicationConditionals::SetAutonomousAsSimulated(uint32 ConnectionId, bool bEnabled)
+{
+	if (ConnectionInfos.IsValidIndex(ConnectionId))
+	{
+		FPerConnectionInfo& ConnectionInfo = ConnectionInfos[ConnectionId];
+		ConnectionInfo.bAutonomousAsSimulated = bEnabled;
+	}
+	else
+	{
+		UE_LOG(LogIris, Warning, TEXT("Attempting to set AutonomousAsSimulated on an invalid connection id %d"), ConnectionId);
+	}
+}
+
 bool FReplicationConditionals::SetCondition(FInternalNetRefIndex ObjectIndex, EReplicationCondition Condition, bool bEnable)
 {
 	if (!ensure(Condition != EReplicationCondition::RoleAutonomous))
@@ -837,10 +850,19 @@ FReplicationConditionals::FConditionalsMask FReplicationConditionals::GetLifetim
 	const bool bIsReplicatingToOwner = (ReplicatingConnectionId == ObjectOwnerConnectionId);
 
 	const FReplicationConditionals::FPerObjectInfo* ObjectInfo = GetPerObjectInfo(ParentObjectIndex);
-	const bool bRoleSimulated = ReplicatingConnectionId != ObjectInfo->AutonomousConnectionId;
+	bool bRoleSimulated = ReplicatingConnectionId != ObjectInfo->AutonomousConnectionId;
 	const bool bRoleAutonomous = ReplicatingConnectionId == ObjectInfo->AutonomousConnectionId;
 	const bool bRepPhysics = ObjectInfo->bRepPhysics;
 
+	if (ConnectionInfos.IsValidIndex(ObjectOwnerConnectionId))
+	{
+		const FPerConnectionInfo& ConnectionInfo = ConnectionInfos[ObjectOwnerConnectionId];
+		if (bRoleAutonomous && ConnectionInfo.bAutonomousAsSimulated)
+		{
+			bRoleSimulated = true;
+		}
+	}
+
 	ConditionalsMask.SetConditionEnabled(COND_None, true);
 	ConditionalsMask.SetConditionEnabled(COND_Custom, true);
 	ConditionalsMask.SetConditionEnabled(COND_Dynamic, true);
diff --git a/Engine/Source/Runtime/Net/Iris/Private/Iris/ReplicationSystem/Conditionals/ReplicationConditionals.h b/Engine/Source/Runtime/Net/Iris/Private/Iris/ReplicationSystem/Conditionals/ReplicationConditionals.h
--- a/Engine/Source/Runtime/Net/Iris/Private/Iris/ReplicationSystem/Conditionals/ReplicationConditionals.h
+++ b/Engine/Source/Runtime/Net/Iris/Private/Iris/ReplicationSystem/Conditionals/ReplicationConditionals.h
@@ -63,6 +63,17 @@ public:
 	void AddConnection(uint32 ConnectionId);
 	void RemoveConnection(uint32 ConnectionId);
 
+	/** 
+	 * When considering replication conditions that require an owning object's role to be ROLE_SimulatedProxy, and the
+	 * owning object's role is ROLE_AutonomousProxy, treat its role as ROLE_SimulatedProxy if this option is set to true.
+	 * 
+	 * Example:
+	 *	If an object has the role ROLE_AutonomousProxy but it has a property with the condition COND_SimulatedOnly, the
+	 *	property would still be replicated to clients. Other properties that would have COND_AutonomousOnly will still
+	 *	be replicated.
+	 */
+	void SetAutonomousAsSimulated(uint32 ConnectionId, bool bEnabled);
+
 	bool SetConditionConnectionFilter(FInternalNetRefIndex ObjectIndex, EReplicationCondition Condition, uint32 ConnectionId, bool bEnable);
 	bool SetCondition(FInternalNetRefIndex ObjectIndex, EReplicationCondition Condition, bool bEnable);
 	void SetOwningConnection(FInternalNetRefIndex ObjectIndex, uint32 ConnectionId);
@@ -105,6 +116,7 @@ private:
 	struct FPerConnectionInfo
 	{
 		TArray<FConditionalsMask> ObjectConditionals;
+		bool bAutonomousAsSimulated = false;
 	};
 
 	struct FSubObjectConditionInfo
diff --git a/Engine/Source/Runtime/Net/Iris/Private/Iris/ReplicationSystem/ReplicationSystem.cpp b/Engine/Source/Runtime/Net/Iris/Private/Iris/ReplicationSystem/ReplicationSystem.cpp
--- a/Engine/Source/Runtime/Net/Iris/Private/Iris/ReplicationSystem/ReplicationSystem.cpp
+++ b/Engine/Source/Runtime/Net/Iris/Private/Iris/ReplicationSystem/ReplicationSystem.cpp
@@ -1841,6 +1841,11 @@ void UReplicationSystem::SetGroupFilterStatus(FNetObjectGroupHandle GroupHandle,
 	Filtering.SetGroupFilterStatus(GroupHandle, ConnectionId, ReplicationStatus);
 }
 
+void UReplicationSystem::SetProxyConnection(uint32 ConnectionId, bool bEnabled)
+{
+	Impl->ReplicationSystemInternal.GetConditionals().SetAutonomousAsSimulated(ConnectionId, bEnabled);
+}
+
 void UReplicationSystem::SetGroupFilterStatus(FNetObjectGroupHandle GroupHandle, const UE::Net::FNetBitArray& Connections, UE::Net::ENetFilterStatus ReplicationStatus)
 {
 	// Early out if this is invalid group
diff --git a/Engine/Source/Runtime/Net/Iris/Private/Iris/Serialization/NetRoleNetSerializer.cpp b/Engine/Source/Runtime/Net/Iris/Private/Iris/Serialization/NetRoleNetSerializer.cpp
--- a/Engine/Source/Runtime/Net/Iris/Private/Iris/Serialization/NetRoleNetSerializer.cpp
+++ b/Engine/Source/Runtime/Net/Iris/Private/Iris/Serialization/NetRoleNetSerializer.cpp
@@ -40,13 +40,29 @@ UE_NET_IMPLEMENT_SERIALIZER_INTERNAL(FNetRoleNetSerializer);
 
 void FNetRoleNetSerializer::Serialize(FNetSerializationContext& Context, const FNetSerializeArgs& Args)
 {
+	if (Context.IsInitializingDefaultState())
+	{
+		return;
+	}
+
 	const ConfigType* Config = static_cast<const ConfigType*>(Args.NetSerializerConfig);
 	const SourceType Value = *reinterpret_cast<const SourceType*>(Args.Source);
 	if (Context.GetInternalContext()->bAlwaysSwapRolesOnReplication)
 	{
 		FNetSerializeArgs SwapArgs = Args;
-		SwapArgs.Source += Config->RelativeInternalOffsetToOtherRole;
+
+		SourceType* OtherRolePtr = reinterpret_cast<SourceType*>(Args.Source + Config->RelativeInternalOffsetToOtherRole);
+		const SourceType OtherRoleValue = *OtherRolePtr;
+
+		if (Context.GetInternalContext()->bDowngradeAutonomousProxyRole & (OtherRoleValue == Config->AutonomousProxyValue))
+		{
+			*OtherRolePtr = Config->SimulatedProxyValue;
+		}
+
+		SwapArgs.Source = NetSerializerValuePointer(OtherRolePtr);
 		Private::FIntRangeNetSerializerBase<SourceType, ConfigType>::Serialize(Context, SwapArgs);
+
+		*OtherRolePtr = OtherRoleValue;
 	}
 	else if (Context.GetInternalContext()->bDowngradeAutonomousProxyRole & (Value == Config->AutonomousProxyValue))
 	{
@@ -62,6 +78,11 @@ void FNetRoleNetSerializer::Serialize(FNetSerializationContext& Context, const F
 
 void FNetRoleNetSerializer::Deserialize(FNetSerializationContext& Context, const FNetDeserializeArgs& Args)
 {
+	if (Context.IsInitializingDefaultState())
+	{
+		return;
+	}
+
 	Private::FIntRangeNetSerializerBase<SourceType, ConfigType>::Deserialize(Context, Args);
 }
 
diff --git a/Engine/Source/Runtime/Net/Iris/Public/Iris/ReplicationSystem/ReplicationSystem.h b/Engine/Source/Runtime/Net/Iris/Public/Iris/ReplicationSystem/ReplicationSystem.h
--- a/Engine/Source/Runtime/Net/Iris/Public/Iris/ReplicationSystem/ReplicationSystem.h
+++ b/Engine/Source/Runtime/Net/Iris/Public/Iris/ReplicationSystem/ReplicationSystem.h
@@ -595,6 +595,9 @@ public:
 	/** Set status of GroupFilter for specific connection. */
 	IRISCORE_API void SetGroupFilterStatus(FNetObjectGroupHandle GroupHandle, uint32 ConnectionId, UE::Net::ENetFilterStatus ReplicationStatus);
 
+	/** Set to true if a connection's remote endpoint is a proxy. */
+	IRISCORE_API void SetProxyConnection(uint32 ConnectionId, bool bEnabled);
+
 	/**
 	 * Set status of GroupFilter for connection marked in the Connections BitArray to the passed ReplicationStatus,
 	 * Connections not marked in the BitArray will be set to the opposite status.
