Support to Iris for proxying actors between multiple game servers and clients:
* Introduced a frontend and backend proxy type to UReplicationSystem.
* A shared FObjectReferenceCache to bridge the frontend and backend proxies.
* Decouple FNetTokenStore from FObjectReferenceCache.
* Replicate to clients a reason when stopping replication of an object.

diff --git a/Engine/Plugins/Runtime/MultiServerReplication/Source/MultiServerReplication/Private/MultiServerProxy.cpp b/Engine/Plugins/Runtime/MultiServerReplication/Source/MultiServerReplication/Private/MultiServerProxy.cpp
--- a/Engine/Plugins/Runtime/MultiServerReplication/Source/MultiServerReplication/Private/MultiServerProxy.cpp
+++ b/Engine/Plugins/Runtime/MultiServerReplication/Source/MultiServerReplication/Private/MultiServerProxy.cpp
@@ -502,6 +502,7 @@ bool UProxyBackendNetDriver::InitBase(bool bInitAsClient, FNetworkNotify* InNoti
 {
 	// Called before InitBase() because the values can be used during driver initialization.
 	EnableExecuteRPCFunctions(false);
+	SetProxyType(UE::Net::EProxyType::Backend);
 
 	bool bSuccess = Super::InitBase(bInitAsClient, InNotify, URL, bReuseAddressAndPort, Error);
 
@@ -734,6 +735,7 @@ bool UProxyNetDriver::InitBase(bool bInitAsClient, FNetworkNotify* InNotify, con
 	DisableActorLogicAndGameCode();
 	EnableRoleSwappingOfAllActors(true);
 	EnableReplicationOfActorsWithAnyRole(true);
+	SetProxyType(UE::Net::EProxyType::Frontend);
 
 	const bool bSuccess = Super::InitBase(bInitAsClient, ProxyListenerNotify, URL, bReuseAddressAndPort, Error);
 
diff --git a/Engine/Source/Runtime/Engine/Classes/Engine/NetDriver.h b/Engine/Source/Runtime/Engine/Classes/Engine/NetDriver.h
--- a/Engine/Source/Runtime/Engine/Classes/Engine/NetDriver.h
+++ b/Engine/Source/Runtime/Engine/Classes/Engine/NetDriver.h
@@ -20,6 +20,7 @@
 #include "UObject/CoreNet.h"
 #include "Channel.h"
 #include "Net/Core/Connection/ConnectionHandle.h"
+#include "Net/Core/Connection/NetEnums.h"
 #include "Net/Core/Misc/DDoSDetection.h"
 #include "IPAddress.h"
 #include "Net/NetAnalyticsTypes.h"
@@ -1492,6 +1493,9 @@ public:
 	ENGINE_API void EnableReplicationOfActorsWithAnyRole(bool bEnable);
 	ENGINE_API bool IsReplicationOfActorsWithAnyRoleEnabled() const;
 
+	/** (Iris replication only) Set the proxy type of the replication system. */
+	ENGINE_API void SetProxyType(UE::Net::EProxyType InProxyType);
+
 	/** (Generic Replicaton Only) Return true if an actor's role can be downgraded during replication to a given connection.  */
 	virtual bool CanDowngradeActorRole(UNetConnection* Connection, AActor* Actor) const { return true; };
 
@@ -1502,6 +1506,8 @@ private:
 	bool bSwapRolesOfAllActors = false;
 	bool bReplicationOfActorsWithAnyRole = false;
 
+	UE::Net::EProxyType ProxyType = UE::Net::EProxyType::None;
+
 	/** A set of replicators that have queued RPCs for transactions (server only) */
 	TSet< FObjectReplicator* > ReplicatorsWithPendingTransactions;
 
diff --git a/Engine/Source/Runtime/Engine/Private/Actor.cpp b/Engine/Source/Runtime/Engine/Private/Actor.cpp
--- a/Engine/Source/Runtime/Engine/Private/Actor.cpp
+++ b/Engine/Source/Runtime/Engine/Private/Actor.cpp
@@ -1325,7 +1325,9 @@ void AActor::PostMigrate(const FUObjectMigrationContext& MigrationContext)
 			if (World)
 			{
 				// Notify replication systems that this actor has been destroyed.
-				UE::Net::FReplicationSystemUtil::StopReplicatingActor(this, FStopReplicatingActorParams(EEndPlayReason::Destroyed));
+				FStopReplicatingActorParams StopRepParams(EEndPlayReason::Destroyed);
+				StopRepParams.StopReplicatingReason = UE::Net::EStopReplicatingReason::Migrated;
+				UE::Net::FReplicationSystemUtil::StopReplicatingActor(this, StopRepParams);
 
 				UE::Net::Private::FNetPropertyConditionManager::Get().NotifyObjectDestroyed(this);
 
@@ -4725,6 +4727,12 @@ void AActor::DispatchBeginPlay(bool bFromLevelStreaming)
 		{
 			BeginPlay();
 		}
+#if UE_SUPPORT_FOR_ACTOR_TICK_DISABLE
+		else
+		{
+			ActorHasBegunPlay = EActorBeginPlayState::HasBegunPlay;
+		}
+#endif
 
 		ensure(BeginPlayCallDepth - 1 == CurrentCallDepth);
 		BeginPlayCallDepth = CurrentCallDepth;
diff --git a/Engine/Source/Runtime/Engine/Private/Net/Iris/ReplicationSystem/EngineReplicationBridge.cpp b/Engine/Source/Runtime/Engine/Private/Net/Iris/ReplicationSystem/EngineReplicationBridge.cpp
--- a/Engine/Source/Runtime/Engine/Private/Net/Iris/ReplicationSystem/EngineReplicationBridge.cpp
+++ b/Engine/Source/Runtime/Engine/Private/Net/Iris/ReplicationSystem/EngineReplicationBridge.cpp
@@ -439,7 +439,16 @@ UE::Net::FNetRefHandle UEngineReplicationBridge::StartReplicatingActor(AActor* A
 
 	if (!ActorRefHandle.IsValid())
 	{
-		ensureMsgf(false, TEXT("Failed to create NetRefHandle for Actor Named %s"), ToCStr(Actor->GetName()));
+		const bool bIsFrontendProxy = (GetReplicationSystem()->GetProxyType() == EProxyType::Frontend);
+
+		// If the replication system is a frontend proxy this condition can be encountered if attempting to start replicating
+		// an actor that isn't found in the shared object reference cache (i.e. it wasn't replicated to a corresponding backend
+		// proxy cache).
+		UE_CLOG(bIsFrontendProxy, LogIris, Warning, TEXT("Failed to create NetRefHandle for Actor Named %s"), ToCStr(Actor->GetName()));
+
+		// If the replication system isn't a frontend proxy then this condition is unexpected and will ensure.
+		ensureMsgf(bIsFrontendProxy, TEXT("Failed to create NetRefHandle for Actor Named %s"), ToCStr(Actor->GetName()));
+
 		return FNetRefHandle::GetInvalid();
 	}
 
@@ -636,10 +645,15 @@ EEndReplicationFlags UEngineReplicationBridge::ConvertToEndReplicationFlags(UE::
 	using namespace UE::Net;
 
 	EEndReplicationFlags EndReplicationFlags = EEndReplicationFlags::None;
-	if (StopReplicatingReason == EStopReplicatingReason::Destroyed)
+	if (StopReplicatingReason == EStopReplicatingReason::Destroyed || StopReplicatingReason == EStopReplicatingReason::Migrated)
 	{
 		EndReplicationFlags |= EEndReplicationFlags::Destroy | EEndReplicationFlags::DestroyNetHandle | EEndReplicationFlags::ClearNetPushId;
 	}
+	
+	if (StopReplicatingReason == EStopReplicatingReason::Migrated)
+	{
+		EndReplicationFlags |= EEndReplicationFlags::Migrated;
+	}
 
 	// If we are shutting down we do not need to validate that we are not detaching remote instances by accident.
 	const bool bIsShuttingDown = (StopReplicatingReason == EStopReplicatingReason::EndPlayInEditor) || (StopReplicatingReason == EStopReplicatingReason::Quit) || NetDriver->GetPendingDestruction();
diff --git a/Engine/Source/Runtime/Engine/Private/NetDriver.cpp b/Engine/Source/Runtime/Engine/Private/NetDriver.cpp
--- a/Engine/Source/Runtime/Engine/Private/NetDriver.cpp
+++ b/Engine/Source/Runtime/Engine/Private/NetDriver.cpp
@@ -436,6 +436,13 @@ namespace UE::Net
 #endif
 
 // CVars
+int32 GInitialNetworkId = -1;
+static FAutoConsoleVariableRef CVarInitialNetworkId(
+	TEXT("net.InitialNetworkId"),
+	GInitialNetworkId,
+	TEXT("The initial value to use when assigning network ids (FNetRefHandle or FNetGUID). If -1 then not explicitly set and the default for the replication system is used."),
+	ECVF_Default);
+
 int32 GSetNetDormancyEnabled = 1;
 static FAutoConsoleVariableRef CVarSetNetDormancyEnabled(
 	TEXT("net.DormancyEnable"),
@@ -770,7 +777,14 @@ void UNetDriver::PostInitProperties()
 		RemoteRoleProperty	= FindFieldChecked<FProperty>( AActor::StaticClass(), TEXT("RemoteRole") );
 
 		PRAGMA_DISABLE_DEPRECATION_WARNINGS
-		GuidCache			= TSharedPtr< FNetGUIDCache >( new FNetGUIDCache( this ) );
+		if (GInitialNetworkId >= 0)
+		{
+			GuidCache = TSharedPtr<FNetGUIDCache>(new FNetGUIDCache(this, GInitialNetworkId));
+		}
+		else
+		{
+			GuidCache = TSharedPtr<FNetGUIDCache>(new FNetGUIDCache(this));
+		}
 		PRAGMA_ENABLE_DEPRECATION_WARNINGS
 
 		NetCache			= TSharedPtr< FClassNetCacheMgr >( new FClassNetCacheMgr() );
@@ -2320,6 +2334,11 @@ bool UNetDriver::IsReplicationOfActorsWithAnyRoleEnabled() const
 	return bReplicationOfActorsWithAnyRole;
 }
 
+void UNetDriver::SetProxyType(UE::Net::EProxyType InProxyType)
+{
+	ProxyType = InProxyType;
+}
+
 ENetMode UNetDriver::GetNetMode() const
 {
 	// Special case for PIE and game worlds - forcing dedicated server behavior if requested
@@ -7476,6 +7495,12 @@ void UNetDriver::CreateReplicationSystem(bool bInitAsClient)
 		Params.bAllowExecuteRepNotifies = !ShouldSkipRepNotifies();
 		Params.bAlwaysSwapRolesOnReplication = IsRoleSwappingOnAllActorsEnabled();
 		Params.NetTokenStore = NetTokenStore.Get();
+		Params.ProxyType = ProxyType;
+
+		if (GInitialNetworkId > 0)
+		{
+			Params.InitialNetRefHandleIndex = GInitialNetworkId;
+		}
 
 		SetReplicationSystem(UE::Net::FReplicationSystemFactory::CreateReplicationSystem(Params));
 	}
diff --git a/Engine/Source/Runtime/Engine/Public/Net/Iris/ReplicationSystem/EngineReplicationBridge.h b/Engine/Source/Runtime/Engine/Public/Net/Iris/ReplicationSystem/EngineReplicationBridge.h
--- a/Engine/Source/Runtime/Engine/Public/Net/Iris/ReplicationSystem/EngineReplicationBridge.h
+++ b/Engine/Source/Runtime/Engine/Public/Net/Iris/ReplicationSystem/EngineReplicationBridge.h
@@ -38,6 +38,8 @@ enum class EStopReplicatingReason : uint32
 	RemovedFromWorld,
 	/** When the application is being exited. */
 	Quit,
+	/** When an actor is moved from one game server to another. */
+	Migrated,
 };
 
 ENGINE_API const TCHAR* LexToString(EStopReplicatingReason Reason);
diff --git a/Engine/Source/Runtime/Net/Core/Public/Net/Core/Connection/NetEnums.h b/Engine/Source/Runtime/Net/Core/Public/Net/Core/Connection/NetEnums.h
--- a/Engine/Source/Runtime/Net/Core/Public/Net/Core/Connection/NetEnums.h
+++ b/Engine/Source/Runtime/Net/Core/Public/Net/Core/Connection/NetEnums.h
@@ -123,6 +123,15 @@ namespace ESecurityEvent
 namespace UE::Net
 {
 
+enum class EProxyType : uint8
+{
+	None = 0,
+	/** Replicate objects from different game servers into the world. */
+	Backend,
+	/** Replicate objects to clients that were replicated into the world by another replication system. */
+	Frontend,
+};
+
 /**
  * Error types used when reporting an error on a specific Iris handle
  */
diff --git a/Engine/Source/Runtime/Net/Iris/Private/Iris/ReplicationSystem/NetBlob/NetBlobManager.cpp b/Engine/Source/Runtime/Net/Iris/Private/Iris/ReplicationSystem/NetBlob/NetBlobManager.cpp
--- a/Engine/Source/Runtime/Net/Iris/Private/Iris/ReplicationSystem/NetBlob/NetBlobManager.cpp
+++ b/Engine/Source/Runtime/Net/Iris/Private/Iris/ReplicationSystem/NetBlob/NetBlobManager.cpp
@@ -86,7 +86,7 @@ bool FNetBlobManager::QueueNetObjectAttachment(uint32 ConnectionId, const FNetOb
 		//This probably happens when the outer list looks like this: Actor->ActorComponent->SubObject1->SubObject2. 
 
 		// If the TargetRef is a valid reference but not replicated, then the outer must be used instead
-		OwnerInfo.CallerRef = ObjectReferenceCache->GetReplicatedOuter(TargetRef);
+		OwnerInfo.CallerRef = ObjectReferenceCache->GetReplicatedOuter(TargetRef, ReplicationSystem);
 
 		// Can that outer send RPCs
 		bCanSendRpc = GetRootObjectAndSubObjectIndicesFromAnyHandle(OwnerInfo.CallerRef.GetRefHandle(), OwnerInfo.RootObjectIndex, OwnerInfo.SubObjectIndex);
@@ -251,7 +251,7 @@ bool FNetBlobManager::GetRPCOwner(FRPCOwner& OutOwnerInfo, const FSendRPCContext
 	// If a root object is sending an RPC
 	if (Context.SubObject == nullptr)
 	{
-		OutOwnerInfo.TargetRef = ObjectReferenceCache->GetOrCreateObjectReference(Context.RootObject);
+		OutOwnerInfo.TargetRef = ObjectReferenceCache->GetOrCreateObjectReference(Context.RootObject, ReplicationSystem);
 		OutOwnerInfo.CallerRef = OutOwnerInfo.TargetRef;
 
 		bCanSendRpc = GetRootObjectIndicesFromHandle(OutOwnerInfo.TargetRef.GetRefHandle(), OutOwnerInfo.RootObjectIndex);
@@ -277,7 +277,7 @@ bool FNetBlobManager::GetRPCOwner(FRPCOwner& OutOwnerInfo, const FSendRPCContext
 		// If the subobject can be referenced
 		if (SubObjectNetRef.IsValid())
 		{
-			OutOwnerInfo.TargetRef = ObjectReferenceCache->GetOrCreateObjectReference(Context.SubObject);
+			OutOwnerInfo.TargetRef = ObjectReferenceCache->GetOrCreateObjectReference(Context.SubObject, ReplicationSystem);
 			OutOwnerInfo.CallerRef = OutOwnerInfo.TargetRef;
 
 			check(OutOwnerInfo.TargetRef.GetRefHandle() == SubObjectNetRef);
@@ -288,8 +288,8 @@ bool FNetBlobManager::GetRPCOwner(FRPCOwner& OutOwnerInfo, const FSendRPCContext
 		// Send the RPC via the Root object if the subobject is not capable
 		if (!bCanSendRpc)
 		{
-			OutOwnerInfo.TargetRef = ObjectReferenceCache->GetOrCreateObjectReference(Context.SubObject);
-			OutOwnerInfo.CallerRef = ObjectReferenceCache->GetOrCreateObjectReference(Context.RootObject);
+			OutOwnerInfo.TargetRef = ObjectReferenceCache->GetOrCreateObjectReference(Context.SubObject, ReplicationSystem);
+			OutOwnerInfo.CallerRef = ObjectReferenceCache->GetOrCreateObjectReference(Context.RootObject, ReplicationSystem);
 
 			bCanSendRpc = GetRootObjectIndicesFromHandle(OutOwnerInfo.CallerRef.GetRefHandle(), OutOwnerInfo.RootObjectIndex);
 		}
diff --git a/Engine/Source/Runtime/Net/Iris/Private/Iris/ReplicationSystem/NetRefHandleManager.cpp b/Engine/Source/Runtime/Net/Iris/Private/Iris/ReplicationSystem/NetRefHandleManager.cpp
--- a/Engine/Source/Runtime/Net/Iris/Private/Iris/ReplicationSystem/NetRefHandleManager.cpp
+++ b/Engine/Source/Runtime/Net/Iris/Private/Iris/ReplicationSystem/NetRefHandleManager.cpp
@@ -44,6 +44,16 @@ FNetRefHandleManager::FNetRefHandleManager(FReplicationProtocolManager& InReplic
 
 void FNetRefHandleManager::Init(const FInitParams& InitParams)
 {
+	if (InitParams.InitialNetRefHandleIndex > 0)
+	{
+		NextStaticHandleIndex = InitParams.InitialNetRefHandleIndex;
+	}
+
+	if (InitParams.InitialNetRefHandleIndex > 0)
+	{
+		NextDynamicHandleIndex = InitParams.InitialNetRefHandleIndex;
+	}
+
 	MaxActiveObjectCount = FNetBitArray::RoundUpToMaxWordBitCount(InitParams.MaxActiveObjectCount);
 	InternalNetRefIndexGrowSize = InitParams.InternalNetRefIndexGrowSize > 0 ? FNetBitArray::RoundUpToMaxWordBitCount(InitParams.InternalNetRefIndexGrowSize) : MaxActiveObjectCount;
 
diff --git a/Engine/Source/Runtime/Net/Iris/Private/Iris/ReplicationSystem/NetRefHandleManager.h b/Engine/Source/Runtime/Net/Iris/Private/Iris/ReplicationSystem/NetRefHandleManager.h
--- a/Engine/Source/Runtime/Net/Iris/Private/Iris/ReplicationSystem/NetRefHandleManager.h
+++ b/Engine/Source/Runtime/Net/Iris/Private/Iris/ReplicationSystem/NetRefHandleManager.h
@@ -67,6 +67,16 @@ enum class ERemoveDependentObjectFlags : uint32
 };
 ENUM_CLASS_FLAGS(ERemoveDependentObjectFlags);
 
+/** The reason that is replicated to clients when an on object is removed from the replication system. */
+enum class ECloseReplicationReason : uint8
+{
+	None = 0,
+	/** When an actor is migrated from one game server to another. */
+	Migrated = 1U,
+	/** The number of bits required to hold the largest value above. */
+	MaxBits = 1U
+};
+
 // Internal class to manage NetHandles and their internal data
 class FNetRefHandleManager
 {
@@ -119,6 +129,8 @@ public:
 				uint16 bHasCachedCreationInfo : 1U;
 			};
 		};
+
+		ECloseReplicationReason CloseReason = ECloseReplicationReason::None;
 	
 		/** Returns true if this is a SubObject that gets destroyed when his RootObject is destroyed too. */
 		bool IsOwnedSubObject() const { return SubObjectRootIndex != InvalidInternalIndex && !!bDestroySubObjectWithOwner; }
@@ -146,6 +158,7 @@ public:
 		uint32 InternalNetRefIndexInitSize = 0;
 		uint32 NetChunkedArrayCount = 0;
 		uint32 InternalNetRefIndexGrowSize = 0;
+		uint32 InitialNetRefHandleIndex = 1;
 	};
 
 	FNetRefHandleManager(FReplicationProtocolManager& InReplicationProtocolManager);
diff --git a/Engine/Source/Runtime/Net/Iris/Private/Iris/ReplicationSystem/ObjectReferenceCache.cpp b/Engine/Source/Runtime/Net/Iris/Private/Iris/ReplicationSystem/ObjectReferenceCache.cpp
--- a/Engine/Source/Runtime/Net/Iris/Private/Iris/ReplicationSystem/ObjectReferenceCache.cpp
+++ b/Engine/Source/Runtime/Net/Iris/Private/Iris/ReplicationSystem/ObjectReferenceCache.cpp
@@ -31,6 +31,7 @@
 #include "Misc/PackageName.h"
 #include "Misc/StringBuilder.h"
 
+#include "Net/Core/Connection/NetEnums.h"
 #include "Net/Core/Trace/NetDebugName.h"
 #include "Net/Core/Trace/NetTrace.h"
 
@@ -112,26 +113,16 @@ void FObjectReferenceCache::FPendingAsyncLoadRequest::Merge(FNetRefHandle InNetR
 static const int INTERNAL_READ_REF_RECURSION_LIMIT = 16;
 
 FObjectReferenceCache::FObjectReferenceCache()
-: ReplicationSystem(nullptr)
-, ReplicationBridge(nullptr)
-, NetTokenStore(nullptr)
-, StringTokenStore(nullptr)
-, NetRefHandleManager(nullptr)
-, bIsAuthority(false)
+: bIsAuthority(false)
 , AsyncLoadMode(EAsyncLoadMode::UseCVar)
 , bCachedCVarAllowAsyncLoading(false)
 {
 	SetAsyncLoadMode(EAsyncLoadMode::UseCVar);
 }
 
-void FObjectReferenceCache::Init(UReplicationSystem* InReplicationSystem)
+void FObjectReferenceCache::Init(bool bAuthority)
 {
-	ReplicationSystem = InReplicationSystem;
-	ReplicationBridge = InReplicationSystem->GetReplicationBridgeAs<UObjectReplicationBridge>();
-	NetTokenStore =  InReplicationSystem->GetNetTokenStore();
-	StringTokenStore = NetTokenStore->GetDataStore<FStringTokenStore>();
-	NetRefHandleManager = &InReplicationSystem->GetReplicationSystemInternal()->GetNetRefHandleManager();
-	bIsAuthority = InReplicationSystem->IsServer();
+	bIsAuthority = bAuthority;
 }
 
 bool FObjectReferenceCache::SupportsObjectInternal(const UObject* Object) const
@@ -244,25 +235,26 @@ bool FObjectReferenceCache::ShouldIgnoreWhenMissing(FNetRefHandle RefHandle) con
 	return CacheObject->bIgnoreWhenMissing;
 }
 
-bool FObjectReferenceCache::RenamePathForPie(uint32 ConnectionId, FString& Str, bool bReading)
+bool FObjectReferenceCache::RenamePathForPie(uint32 ConnectionId, FString& Str, bool bReading, UReplicationSystem* ReplicationSystem)
 {
+	UObjectReplicationBridge* ReplicationBridge = ReplicationSystem->GetReplicationBridgeAs<UObjectReplicationBridge>();
 	return UE::GetPlayInEditorID() != -1 ? ReplicationBridge->RemapPathForPIE(ConnectionId, Str, bReading) : false;
 }
 
-FNetRefHandle FObjectReferenceCache::CreateObjectReferenceHandle(const UObject* Object)
+FNetRefHandle FObjectReferenceCache::CreateObjectReferenceHandle(const UObject* Object, UReplicationSystem* ReplicationSystem)
 {
 	FNetObjectReference Ref;
-	CreateObjectReferenceInternal(Object, Ref);
+	CreateObjectReferenceInternal(Object, Ref, ReplicationSystem);
 
 	return Ref.GetRefHandle();
 }
 
-FNetRefHandle FObjectReferenceCache::PreRegisterObjectReferenceHandle(const UObject* Object)
+FNetRefHandle FObjectReferenceCache::PreRegisterObjectReferenceHandle(const UObject* Object, UReplicationSystem* ReplicationSystem)
 {
 	if (IsAuthority())
 	{
 		FNetObjectReference Ref;
-		CreateObjectReferenceInternal(Object, Ref);
+		CreateObjectReferenceInternal(Object, Ref, ReplicationSystem);
 		FCachedNetObjectReference* CachedRef = ReferenceHandleToCachedReference.Find(Ref.GetRefHandle());
 		if (CachedRef)
 		{
@@ -280,12 +272,14 @@ FNetRefHandle FObjectReferenceCache::PreRegisterObjectReferenceHandle(const UObj
 	return FNetRefHandle::GetInvalid();
 }
 
-bool FObjectReferenceCache::CreateObjectReferenceInternal(const UObject* Object, FNetObjectReference& OutReference)
+bool FObjectReferenceCache::CreateObjectReferenceInternal(const UObject* Object, FNetObjectReference& OutReference, UReplicationSystem* ReplicationSystem)
 {
 	if (!IsValid(Object))
 	{
 		return false;
 	}
+	
+	FStringTokenStore* LocalStringTokenStore = ReplicationSystem->GetNetTokenStore()->GetDataStore<FStringTokenStore>();
 
 	CA_ASSUME(Object != nullptr);
 
@@ -381,8 +375,8 @@ bool FObjectReferenceCache::CreateObjectReferenceInternal(const UObject* Object,
 				FString RelativeObjectPath = Object->GetPathName(Outer);
 				constexpr bool bReading = false;
 				// The connection ID isn't used unless reading.
-				RenamePathForPie(UE::Net::InvalidConnectionId, RelativeObjectPath, bReading);
-				OutReference = MakeNetObjectReference(*ReplicatedOuter, StringTokenStore->GetOrCreateToken(RelativeObjectPath));
+				RenamePathForPie(UE::Net::InvalidConnectionId, RelativeObjectPath, bReading, ReplicationSystem);
+				OutReference = MakeNetObjectReference(*ReplicatedOuter, LocalStringTokenStore->GetOrCreateToken(RelativeObjectPath));
 
 				return true;
 			}
@@ -394,8 +388,8 @@ bool FObjectReferenceCache::CreateObjectReferenceInternal(const UObject* Object,
 			FString ObjectPath = Object->GetPathName();
 			constexpr bool bReading = false;
 			// The connection ID isn't used unless reading.
-			RenamePathForPie(UE::Net::InvalidConnectionId, ObjectPath, bReading);
-			OutReference = MakeNetObjectReference(FNetRefHandle::GetInvalid(), StringTokenStore->GetOrCreateToken(ObjectPath));	
+			RenamePathForPie(UE::Net::InvalidConnectionId, ObjectPath, bReading, ReplicationSystem);
+			OutReference = MakeNetObjectReference(FNetRefHandle::GetInvalid(), LocalStringTokenStore->GetOrCreateToken(ObjectPath));	
 			
 			return true;
 		}
@@ -405,7 +399,35 @@ bool FObjectReferenceCache::CreateObjectReferenceInternal(const UObject* Object,
 	}
 
 	// Generate a new RefHandle
-	FNetRefHandle NetRefHandle = NetRefHandleManager->AllocateNetRefHandle(bIsStatic);
+	FNetRefHandle NetRefHandle; 
+
+	FNetRefHandleManager& NetRefHandleManager = ReplicationSystem->GetReplicationSystemInternal()->GetNetRefHandleManager();
+	
+	if (ReplicationSystem->GetProxyType() == UE::Net::EProxyType::Frontend)
+	{
+		TSharedPtr<FObjectReferenceCache> SharedObjectReferenceCache = FReplicationSystemFactory::GetSharedObjectReferenceCache();
+		if (!ensureMsgf(SharedObjectReferenceCache, TEXT("Failed initializing and retrieving the shared obejct reference cache.")))
+		{
+			return false;
+		}
+
+		FNetRefHandle SharedNetRefHandle = SharedObjectReferenceCache->GetObjectReferenceHandleFromObject(Object); 
+		if (SharedNetRefHandle.IsValid())
+		{
+			UE_LOG(LogIris, Verbose, TEXT("ObjectReferenceCache::CreateObjectReferenceHandle Found a valid NetRefHandle %s in the shared object reference cache for object %s"), *SharedNetRefHandle.ToString(), *Object->GetName());
+			NetRefHandle = NetRefHandleManager.MakeNetRefHandle(SharedNetRefHandle.GetId(), ReplicationSystem->GetId());
+		}
+		else
+		{
+			UE_LOG(LogIris, Warning, TEXT("ObjectReferenceCache::CreateObjectReferenceHandle Unable to find a NetRefHandle for object %s in the shared object reference cache."), *Object->GetName());			
+			return false;
+		}
+	}
+	else
+	{
+		NetRefHandle = NetRefHandleManager.AllocateNetRefHandle(bIsStatic);
+	}
+
 	check(NetRefHandle.IsValid());
 
 	FNetToken PathToken;	
@@ -416,13 +438,13 @@ bool FObjectReferenceCache::CreateObjectReferenceInternal(const UObject* Object,
 		FString ObjectPath = Object->GetName();
 		constexpr bool bReading = false;
 		// The connection ID isn't used unless reading.
-		RenamePathForPie(UE::Net::InvalidConnectionId, ObjectPath, bReading);
+		RenamePathForPie(UE::Net::InvalidConnectionId, ObjectPath, bReading, ReplicationSystem);
 
 #if WITH_EDITOR
 		ensureMsgf(!ObjectPath.IsEmpty(), TEXT("NetworkRemapPath found %s to be an invalid name for %s. This object will not replicate!"), *Object->GetName(), *GetPathNameSafe(Object));
 #endif
 
-		PathToken = StringTokenStore->GetOrCreateToken(ObjectPath);
+		PathToken = LocalStringTokenStore->GetOrCreateToken(ObjectPath);
 	}
 
 	// Register cached ref
@@ -432,7 +454,7 @@ bool FObjectReferenceCache::CreateObjectReferenceInternal(const UObject* Object,
 	CachedObject.ObjectKey = Object;
 	CachedObject.NetRefHandle = NetRefHandle;
 	CachedObject.RelativePath = PathToken;
-	CachedObject.OuterNetRefHandle = CreateObjectReferenceHandle(Object->GetOuter());
+	CachedObject.OuterNetRefHandle = CreateObjectReferenceHandle(Object->GetOuter(), ReplicationSystem);
 	CachedObject.bIsPackage = bIsStatic && !Object->GetOuter();
 	CachedObject.bNoLoad = !CanClientLoadObjectInternal(Object, bIsDynamic); // We are probably going to deal with this as a dependency
 	CachedObject.bIgnoreWhenMissing = !CachedObject.bNoLoad;
@@ -449,7 +471,7 @@ bool FObjectReferenceCache::CreateObjectReferenceInternal(const UObject* Object,
 	if (NetRefHandle.IsValid())
 	{
 		TStringBuilder<512> Builder;
-		Builder.Appendf(TEXT("Ref_%s"), PathToken.IsValid() ? StringTokenStore->ResolveToken(PathToken) : ToCStr(Object->GetName()));
+		Builder.Appendf(TEXT("Ref_%s"), PathToken.IsValid() ? LocalStringTokenStore->ResolveToken(PathToken) : ToCStr(Object->GetName()));
 		UE_NET_TRACE_NETHANDLE_CREATED(NetRefHandle, CreatePersistentNetDebugName(Builder.ToString()), 0, IsAuthority() ? 1U : 0U);
 	}
 #endif
@@ -569,7 +591,7 @@ void FObjectReferenceCache::BindRemoteReference(FNetRefHandle RefHandle, const U
 	ObjectToNetReferenceHandle.Add(Object, RefHandle);
 }
 
-void FObjectReferenceCache::AddPreRegisteredReference(FNetRefHandle RefHandle, const UObject* Object)
+void FObjectReferenceCache::AddPreRegisteredReference(FNetRefHandle RefHandle, const UObject* Object, UReplicationSystem* ReplicationSystem)
 {
 	if (!IsAuthority())
 	{
@@ -785,7 +807,7 @@ UObject* FObjectReferenceCache::ResolveObjectReferenceHandleInternal(FNetRefHand
 		{			
 			FString ObjectPath(StringTokenStore->ResolveRemoteToken(CacheObjectPtr->RelativePath, *ResolveContext.RemoteNetTokenStoreState));
 			constexpr bool bReading = true;
-			RenamePathForPie(ResolveContext.ConnectionId, ObjectPath, bReading);
+			RenamePathForPie(ResolveContext.ConnectionId, ObjectPath, bReading, ResolveContext.ReplicationSystem);
 			check(ObjectPath == Object->GetFName().ToString() || RefHandle.IsDynamic() || IsAuthority());
 
 			FNetRefHandle* ExistingRefHandle = ObjectToNetReferenceHandle.Find(Object);
@@ -872,7 +894,8 @@ UObject* FObjectReferenceCache::ResolveObjectReferenceHandleInternal(FNetRefHand
 		UE_LOG(LogIris, Error, TEXT("GetObjectFromRefHandle: Outer is null or package is not fully loaded. FullPath: %s Outer: %s"), ToCStr(FullPath(RefHandle, ResolveContext)), *GetFullNameSafe(ObjOuter));
 	}
 
-	const TCHAR* ResolvedToken = StringTokenStore->ResolveRemoteToken(CacheObjectPtr->RelativePath, *ResolveContext.RemoteNetTokenStoreState);
+	FStringTokenStore* LocalStringTokenStore = ResolveContext.ReplicationSystem->GetNetTokenStore()->GetDataStore<FStringTokenStore>();
+	const TCHAR* ResolvedToken = LocalStringTokenStore->ResolveRemoteToken(CacheObjectPtr->RelativePath, *ResolveContext.RemoteNetTokenStoreState);
 	if (ResolvedToken == nullptr)
 	{
 		return nullptr;
@@ -880,7 +903,7 @@ UObject* FObjectReferenceCache::ResolveObjectReferenceHandleInternal(FNetRefHand
 
 	FString ObjectPath(ResolvedToken);
 	constexpr bool bReading = true;
-	RenamePathForPie(ResolveContext.ConnectionId, ObjectPath, bReading);
+	RenamePathForPie(ResolveContext.ConnectionId, ObjectPath, bReading, ResolveContext.ReplicationSystem);
 
 #if WITH_EDITOR
 	ensureMsgf(!ObjectPath.IsEmpty(), TEXT("NetworkRemapPath found %s to be an invalid name. This object will not be binded and replicated!"), ResolvedToken);
@@ -923,7 +946,7 @@ UObject* FObjectReferenceCache::ResolveObjectReferenceHandleInternal(FNetRefHand
 				{
 					CacheObjectPtr->bIsPending = true;
 
-					StartAsyncLoadingPackage(*CacheObjectPtr, ObjectPathName, ResolveContext.AsyncLoadingPriority, RefHandle, false);
+					StartAsyncLoadingPackage(*CacheObjectPtr, ObjectPathName, ResolveContext.AsyncLoadingPriority, RefHandle, false, ResolveContext.ReplicationSystem);
 					UE_LOG_REFERENCECACHE(VeryVerbose, TEXT("GetObjectFromRefHandle: Async loading package. Path: %s, %s"), *ObjectPath, *RefHandle.ToString());
 				}
 				else
@@ -985,7 +1008,7 @@ UObject* FObjectReferenceCache::ResolveObjectReferenceHandleInternal(FNetRefHand
 			if (ShouldAsyncLoad() && !ResolveContext.bForceSyncLoad && Package->HasAnyInternalFlags(EInternalObjectFlags_AsyncLoading))
 			{
 				// Something else is already async loading this package, calling load again will add our callback to the existing load request
-				StartAsyncLoadingPackage(*CacheObjectPtr, ObjectPathName, ResolveContext.AsyncLoadingPriority, RefHandle, true);
+				StartAsyncLoadingPackage(*CacheObjectPtr, ObjectPathName, ResolveContext.AsyncLoadingPriority, RefHandle, true, ResolveContext.ReplicationSystem);
 				UE_LOG(LogIrisReferences, Verbose, TEXT("GetObjectFromRefHandle: Listening to existing async load. Path: %s, NetRefHandle: %s"), *ObjectPath, *RefHandle.ToString());
 			}
 			else if (/*!GbNetCheckNoLoadPackages ||*/ !CacheObjectPtr->bNoLoad)
@@ -1025,6 +1048,7 @@ UObject* FObjectReferenceCache::ResolveObjectReferenceHandleInternal(FNetRefHand
 	//	}
 	//}
 
+	UObjectReplicationBridge* ReplicationBridge = ResolveContext.ReplicationSystem->GetReplicationBridgeAs<UObjectReplicationBridge>();
 	if (Object && !ReplicationBridge->ObjectLevelHasFinishedLoading(Object))
 	{
 		UE_LOG_REFERENCECACHE(Verbose, TEXT("GetObjectFromRefHandle: Forcing object to NULL since level is not loaded yet. Object: %s" ), *Object->GetFullName());
@@ -1073,7 +1097,7 @@ UObject* FObjectReferenceCache::ResolveObjectReferenceHandleInternal(FNetRefHand
 	//}
 
 	// Promote to full handle
-	FNetRefHandle CompleteRefHandle = FNetRefHandleManager::MakeNetRefHandle(RefHandle.GetId(), ReplicationSystem->GetId());
+	FNetRefHandle CompleteRefHandle = FNetRefHandleManager::MakeNetRefHandle(RefHandle.GetId(), ResolveContext.ReplicationSystem->GetId());
 
 	// If we already have an existing handle associated with this object, we need do do a few additional checks.
 	// Due to level streaming (streaming in and out the same level) we might have multiple path based references 
@@ -1136,13 +1160,15 @@ ENetObjectReferenceResolveResult FObjectReferenceCache::ResolveObjectReference(c
 	
 	if (Reference.PathToken.IsValid())
 	{
+		FStringTokenStore* LocalStringTokenStore = ResolveContext.ReplicationSystem->GetNetTokenStore()->GetDataStore<FStringTokenStore>();
+
 		// This path is only used by Client to Server references
-		const TCHAR* ResolvedToken = StringTokenStore->ResolveRemoteToken(Reference.PathToken, *ResolveContext.RemoteNetTokenStoreState);
+		const TCHAR* ResolvedToken = LocalStringTokenStore->ResolveRemoteToken(Reference.PathToken, *ResolveContext.RemoteNetTokenStoreState);
 		if (ResolvedToken)
 		{
 			FString ObjectPath(ResolvedToken);
 			constexpr bool bReading = true;
-			RenamePathForPie(ResolveContext.ConnectionId, ObjectPath, bReading);
+			RenamePathForPie(ResolveContext.ConnectionId, ObjectPath, bReading, ResolveContext.ReplicationSystem);
 
 			// If both path and handle is valid this is a reference to an object with a relative path.
 			if (Reference.GetRefHandle().IsValid())
@@ -1199,7 +1225,7 @@ ENetObjectReferenceResolveResult FObjectReferenceCache::ResolveObjectReference(c
 }
 
 // Find replicated outer
-FNetObjectReference FObjectReferenceCache::GetReplicatedOuter(const FNetObjectReference& Reference) const
+FNetObjectReference FObjectReferenceCache::GetReplicatedOuter(const FNetObjectReference& Reference, UReplicationSystem* ReplicationSystem) const
 {
 	if (!Reference.IsValid())
 	{
@@ -1211,7 +1237,8 @@ FNetObjectReference FObjectReferenceCache::GetReplicatedOuter(const FNetObjectRe
 
 	while (OuterNetRefHandle.IsValid())
 	{
-		if (NetRefHandleManager->GetInternalIndex(OuterNetRefHandle))
+		FNetRefHandleManager& NetRefHandleManager = ReplicationSystem->GetReplicationSystemInternal()->GetNetRefHandleManager();
+		if (NetRefHandleManager.GetInternalIndex(OuterNetRefHandle))
 		{
 			return FNetObjectReference(OuterNetRefHandle);
 		}
@@ -1241,15 +1268,15 @@ FNetRefHandle FObjectReferenceCache::GetDynamicRoot(const FNetRefHandle Handle)
 	}
 }
 
-FNetObjectReference FObjectReferenceCache::GetOrCreateObjectReference(const UObject* Instance)
+FNetObjectReference FObjectReferenceCache::GetOrCreateObjectReference(const UObject* Instance, UReplicationSystem* ReplicationSystem)
 {
 	FNetObjectReference ObjectReference;
-	CreateObjectReferenceInternal(Instance, ObjectReference);
+	CreateObjectReferenceInternal(Instance, ObjectReference, ReplicationSystem);
 
 	return ObjectReference;
 }
 
-FNetObjectReference FObjectReferenceCache::GetOrCreateObjectReference(const FString& ObjectPath, const UObject* Outer)
+FNetObjectReference FObjectReferenceCache::GetOrCreateObjectReference(const FString& ObjectPath, const UObject* Outer, UReplicationSystem* ReplicationSystem)
 {
 	// Only allow creation of path based references on server
 	if (!IsAuthority())
@@ -1261,19 +1288,22 @@ FNetObjectReference FObjectReferenceCache::GetOrCreateObjectReference(const FStr
 	FNetObjectReference OutReference;
 
 	const bool bIsStatic = true;
-	FNetRefHandle RefHandle = NetRefHandleManager->AllocateNetRefHandle(bIsStatic);
+	FNetRefHandleManager& NetRefHandleManager = ReplicationSystem->GetReplicationSystemInternal()->GetNetRefHandleManager();
+	FNetRefHandle RefHandle = NetRefHandleManager.AllocateNetRefHandle(bIsStatic);
 
 	if (RefHandle.IsValid())
 	{
 		LLM_SCOPE_BYTAG(Iris);
 
+		FStringTokenStore* LocalStringTokenStore = ReplicationSystem->GetNetTokenStore()->GetDataStore<FStringTokenStore>();
+
 		// Register cached ref
 		FCachedNetObjectReference CachedObject;
 
 		CachedObject.ObjectKey = nullptr;
 		CachedObject.NetRefHandle = RefHandle;
-		CachedObject.RelativePath = StringTokenStore->GetOrCreateToken(ObjectPath);
-		CachedObject.OuterNetRefHandle = CreateObjectReferenceHandle(Outer);
+		CachedObject.RelativePath = LocalStringTokenStore->GetOrCreateToken(ObjectPath);
+		CachedObject.OuterNetRefHandle = CreateObjectReferenceHandle(Outer, ReplicationSystem);
 		CachedObject.bIsPackage = 0U;
 		CachedObject.bNoLoad = 1U;
 		CachedObject.bIgnoreWhenMissing = 1U;
@@ -1328,9 +1358,12 @@ void FObjectReferenceCache::WriteFullReferenceInternal(FNetSerializationContext&
 
 		if (Writer->WriteBool(bHasPath))
 		{
+			FNetTokenStore* LocalNetTokenStore = Context.GetInternalContext()->ReplicationSystem->GetNetTokenStore();
+			FStringTokenStore* LocalStringTokenStore = Context.GetInternalContext()->ReplicationSystem->GetNetTokenStore()->GetDataStore<FStringTokenStore>();
+
 			// Note: As we use the StringTokenStore to write the token we also need to read it using the stringtokenstore.
-			StringTokenStore->WriteNetToken(Context, CachedObject->RelativePath);
-			NetTokenStore->ConditionalWriteNetTokenData(Context, ExportContext, CachedObject->RelativePath);
+			LocalStringTokenStore->WriteNetToken(Context, CachedObject->RelativePath);
+			LocalNetTokenStore->ConditionalWriteNetTokenData(Context, ExportContext, CachedObject->RelativePath);
 			WriteFullReferenceInternal(Context, FNetObjectReference(CachedObject->OuterNetRefHandle));
 		}
 	
@@ -1397,8 +1430,10 @@ void FObjectReferenceCache::ReadFullReferenceInternal(FNetSerializationContext&
 
 		if (Reader->ReadBool())
 		{
-			RelativePath = StringTokenStore->ReadNetToken(Context);
-			NetTokenStore->ConditionalReadNetTokenData(Context, RelativePath);
+			FStringTokenStore* LocalStringTokenStore = Context.GetInternalContext()->ReplicationSystem->GetNetTokenStore()->GetDataStore<FStringTokenStore>();
+			RelativePath = LocalStringTokenStore->ReadNetToken(Context);
+			FNetTokenStore* LocalNetTokenStore = Context.GetInternalContext()->ReplicationSystem->GetNetTokenStore();
+			LocalNetTokenStore->ConditionalReadNetTokenData(Context, RelativePath);
 			ReadFullReferenceInternal(Context, OuterRef, RecursionCount + 1U);
 		}
 	}
@@ -1483,9 +1518,12 @@ void FObjectReferenceCache::ReadFullReference(FNetSerializationContext& Context,
 	{
 		UE_NET_TRACE_SCOPE(ClientAssigned, *Reader, Context.GetTraceCollector(), ENetTraceVerbosity::Verbose);
 
+		FNetTokenStore* LocalNetTokenStore = Context.GetInternalContext()->ReplicationSystem->GetNetTokenStore();
+		FStringTokenStore* LocalStringTokenStore = Context.GetInternalContext()->ReplicationSystem->GetNetTokenStore()->GetDataStore<FStringTokenStore>();
+
 		const FNetRefHandle NetRefHandle = ReadNetRefHandle(Context);
-		FNetToken RelativePath = StringTokenStore->ReadNetToken(Context);
-		NetTokenStore->ConditionalReadNetTokenData(Context, RelativePath);
+		FNetToken RelativePath = LocalStringTokenStore->ReadNetToken(Context);
+		LocalNetTokenStore->ConditionalReadNetTokenData(Context, RelativePath);
 
 		// We expect a valid RelativePath
 		if (!RelativePath.IsValid())
@@ -1536,10 +1574,13 @@ void FObjectReferenceCache::WriteFullReference(FNetSerializationContext& Context
 	{
 		UE_NET_TRACE_SCOPE(ClientAssigned, *Writer, Context.GetTraceCollector(), ENetTraceVerbosity::Verbose);
 
+		FNetTokenStore* LocalNetTokenStore = Context.GetInternalContext()->ReplicationSystem->GetNetTokenStore();
+		FStringTokenStore* LocalStringTokenStore = Context.GetInternalContext()->ReplicationSystem->GetNetTokenStore()->GetDataStore<FStringTokenStore>();
+
 		WriteNetRefHandle(Context, Ref.GetRefHandle());
 		// Note: As we use the StringTokenStore to write the token we also need to read it using the stringtokenstore.
-		StringTokenStore->WriteNetToken(Context, Ref.PathToken);
-		NetTokenStore->ConditionalWriteNetTokenData(Context, Context.GetExportContext(), Ref.PathToken);
+		LocalStringTokenStore->WriteNetToken(Context, Ref.PathToken);
+		LocalNetTokenStore->ConditionalWriteNetTokenData(Context, Context.GetExportContext(), Ref.PathToken);
 		return;
 	}
 
@@ -1569,9 +1610,11 @@ void FObjectReferenceCache::WriteReference(FNetSerializationContext& Context, FN
 	const bool bIsClientAssignedReference = Ref.PathToken.IsValid();
 	if (Writer->WriteBool(bIsClientAssignedReference))
 	{
+		FStringTokenStore* LocalStringTokenStore = Context.GetInternalContext()->ReplicationSystem->GetNetTokenStore()->GetDataStore<FStringTokenStore>();
+
 		WriteNetRefHandle(Context, Ref.GetRefHandle());
 		// Note: As we use the StringTokenStore to write the token we also need to read it using the stringtokenstore.
-		StringTokenStore->WriteNetToken(Context, Ref.PathToken);
+		LocalStringTokenStore->WriteNetToken(Context, Ref.PathToken);
 
 		return;
 	}
@@ -1606,8 +1649,10 @@ void FObjectReferenceCache::ReadReference(FNetSerializationContext& Context, FNe
 
 	if (const bool bIsClientAssignedReference = Reader->ReadBool())
 	{
+		FStringTokenStore* LocalStringTokenStore = Context.GetInternalContext()->ReplicationSystem->GetNetTokenStore()->GetDataStore<FStringTokenStore>();
+
 		const FNetRefHandle NetRefHandle = ReadNetRefHandle(Context);
-		FNetToken RelativePath = StringTokenStore->ReadNetToken(Context);
+		FNetToken RelativePath = LocalStringTokenStore->ReadNetToken(Context);
 
 		// We expect a valid RelativePath
 		if (!RelativePath.IsValid())
@@ -1727,7 +1772,8 @@ FObjectReferenceCache::EWriteExportsResult FObjectReferenceCache::WritePendingEx
 		TArrayView<const FNetToken> NetTokenExportsView = MakeArrayView(ExportContext->GetBatchExports().NetTokensPendingExportInCurrentBatch);
 		{
 			UE_NET_TRACE_SCOPE(NetTokenExports, Writer, Context.GetTraceCollector(), ENetTraceVerbosity::Verbose);
-			const bool bIsNetTokenAuthority = NetTokenStore->IsAuthority();
+			FNetTokenStore* LocalNetTokenStore = Context.GetInternalContext()->ReplicationSystem->GetNetTokenStore();
+			const bool bIsNetTokenAuthority = LocalNetTokenStore->IsAuthority();
 
 			for (const FNetToken& NetToken : NetTokenExportsView)
 			{
@@ -1736,8 +1782,8 @@ FObjectReferenceCache::EWriteExportsResult FObjectReferenceCache::WritePendingEx
 				{
 					Writer.WriteBool(true);
 
-					NetTokenStore->WriteNetToken(Context, NetToken);
-					NetTokenStore->WriteTokenData(Context, NetToken);
+					LocalNetTokenStore->WriteNetToken(Context, NetToken);
+					LocalNetTokenStore->WriteTokenData(Context, NetToken);
 
 					ExportContext->AddExported(NetToken);
 				}
@@ -1813,8 +1859,10 @@ bool FObjectReferenceCache::ReadExports(const FNetRefHandle& NetObjectHandle, FN
 			FForceInlineExportScope ForceInlineExportScope(Context.GetInternalContext());
 			while (bHasExportsToRead && !Context.HasErrorOrOverflow())
 			{
-				FNetToken ImportedNetToken = NetTokenStore->ReadNetToken(Context);
-				NetTokenStore->ReadTokenData(Context, ImportedNetToken, *ResolveContext.RemoteNetTokenStoreState);
+				FNetTokenStore* LocalNetTokenStore = Context.GetInternalContext()->ReplicationSystem->GetNetTokenStore();
+
+				FNetToken ImportedNetToken = LocalNetTokenStore->ReadNetToken(Context);
+				LocalNetTokenStore->ReadTokenData(Context, ImportedNetToken, *ResolveContext.RemoteNetTokenStoreState);
 				bHasExportsToRead = Reader.ReadBool();
 			}
 		}
@@ -1858,7 +1906,9 @@ bool FObjectReferenceCache::WriteMustBeMappedExports(FNetSerializationContext& C
 		uint32 CountExports = 0;
 #endif
 
-		const EIrisAsyncLoadingPriority IrisAsyncLoadingPriority = NetRefHandleManager->GetIrisAsyncLoadingPriority(ObjectIndex);
+		UReplicationSystem* ReplicationSystem = Context.GetInternalContext()->ReplicationSystem;
+		FNetRefHandleManager& NetRefHandleManager = ReplicationSystem->GetReplicationSystemInternal()->GetNetRefHandleManager();
+		const EIrisAsyncLoadingPriority IrisAsyncLoadingPriority = NetRefHandleManager.GetIrisAsyncLoadingPriority(ObjectIndex);
 
 		// Only write the loading priority before the first reference
 		bool bHasWrittenPriority = false;
@@ -1881,7 +1931,7 @@ bool FObjectReferenceCache::WriteMustBeMappedExports(FNetSerializationContext& C
 				WriteNetRefHandle(Context, Handle);
 
 #if UE_NET_IRIS_CSV_STATS
-				UE_LOG(LogIrisReferences, VeryVerbose, TEXT("WriteMustBeMappedExports for %s. MustBeMapped[%d] %s (%s)"), *NetRefHandleManager->PrintObjectFromIndex(ObjectIndex), CountExports, *GetNameSafe(CachedObject->Object.Get()), *Handle.ToString());
+				UE_LOG(LogIrisReferences, VeryVerbose, TEXT("WriteMustBeMappedExports for %s. MustBeMapped[%d] %s (%s)"), *NetRefHandleManager.PrintObjectFromIndex(ObjectIndex), CountExports, *GetNameSafe(CachedObject->Object.Get()), *Handle.ToString());
 #endif
 
 #if UE_NET_IRIS_CSV_STATS
@@ -1891,7 +1941,7 @@ bool FObjectReferenceCache::WriteMustBeMappedExports(FNetSerializationContext& C
 		}
 
 #if UE_NET_IRIS_CSV_STATS
-		UE_CLOG(CountExports>0, LogIrisReferences, VeryVerbose, TEXT("WriteMustBeMappedExports for %s. Wrote %d exports with priority: %s"), *NetRefHandleManager->PrintObjectFromIndex(ObjectIndex), CountExports, LexToString(IrisAsyncLoadingPriority));
+		UE_CLOG(CountExports>0, LogIrisReferences, VeryVerbose, TEXT("WriteMustBeMappedExports for %s. Wrote %d exports with priority: %s"), *NetRefHandleManager.PrintObjectFromIndex(ObjectIndex), CountExports, LexToString(IrisAsyncLoadingPriority));
 #endif
 
 		// $TODO: This should be a separate state, as MustBeMapped is not really an export, just a list of important exports
@@ -1949,8 +1999,10 @@ FString FObjectReferenceCache::DescribeObjectReference(const FNetObjectReference
 	// Only client assigned FNetObjectReferences has a path stored directly
 	if (Ref.PathToken.IsValid())
 	{
+		FStringTokenStore* LocalStringTokenStore = ResolveContext.ReplicationSystem->GetNetTokenStore()->GetDataStore<FStringTokenStore>();
+
 		// This path is only used by Client to Server references
-		const TCHAR* ResolvedToken = StringTokenStore->ResolveRemoteToken(Ref.PathToken, *ResolveContext.RemoteNetTokenStoreState);
+		const TCHAR* ResolvedToken = LocalStringTokenStore->ResolveRemoteToken(Ref.PathToken, *ResolveContext.RemoteNetTokenStoreState);
 		if (ResolvedToken)
 		{
 			if (!FullPath.IsEmpty())
@@ -1960,7 +2012,7 @@ FString FObjectReferenceCache::DescribeObjectReference(const FNetObjectReference
 
 			FString ObjectPath(ResolvedToken);
 			constexpr bool bReading = true;
-			RenamePathForPie(ResolveContext.ConnectionId, ObjectPath, bReading);
+			RenamePathForPie(ResolveContext.ConnectionId, ObjectPath, bReading, ResolveContext.ReplicationSystem);
 			FullPath += FString::Printf(TEXT("%s"), ResolvedToken);
 		}
 	}
@@ -2024,7 +2076,9 @@ void FObjectReferenceCache::GenerateFullPath_r(FNetRefHandle RefHandle, const FN
 		}
 		else
 		{
-			const TCHAR* ResolvedPath = StringTokenStore->ResolveRemoteToken(CacheObject->RelativePath, *ResolveContext.RemoteNetTokenStoreState);
+			FStringTokenStore* LocalStringTokenStore = ResolveContext.ReplicationSystem->GetNetTokenStore()->GetDataStore<FStringTokenStore>();
+
+			const TCHAR* ResolvedPath = LocalStringTokenStore->ResolveRemoteToken(CacheObject->RelativePath, *ResolveContext.RemoteNetTokenStoreState);
 			FullPath += FString::Printf(TEXT("[%s](%s)"), *RefHandle.ToString(), ResolvedPath);
 		}
 	}
@@ -2060,7 +2114,7 @@ void FObjectReferenceCache::ValidateAsyncLoadingPackage(FCachedNetObjectReferenc
 //#endif
 }
 
-void FObjectReferenceCache::StartAsyncLoadingPackage(FCachedNetObjectReference& CacheObject, FName PackagePath, TAsyncLoadPriority AsyncLoadingPriority, const FNetRefHandle RefHandle, const bool bWasAlreadyAsyncLoading)
+void FObjectReferenceCache::StartAsyncLoadingPackage(FCachedNetObjectReference& CacheObject, FName PackagePath, TAsyncLoadPriority AsyncLoadingPriority, const FNetRefHandle RefHandle, const bool bWasAlreadyAsyncLoading, UReplicationSystem* ReplicationSystem)
 {
 	LLM_SCOPE_BYTAG(Iris);
 
@@ -2071,7 +2125,9 @@ void FObjectReferenceCache::StartAsyncLoadingPackage(FCachedNetObjectReference&
 		AsyncLoadingPriority = ConvertAsyncLoadingPriority(EIrisAsyncLoadingPriority::Default);
 	}
 
-	UE_LOG(LogIrisReferences, Verbose, TEXT("StartAsyncLoadingPackage for %s: Package: %s (Async loading priority: %d)"), *NetRefHandleManager->PrintObjectFromNetRefHandle(RefHandle), *PackagePath.ToString(), AsyncLoadingPriority);
+	FNetRefHandleManager& NetRefHandleManager = ReplicationSystem->GetReplicationSystemInternal()->GetNetRefHandleManager();
+
+	UE_LOG(LogIrisReferences, Verbose, TEXT("StartAsyncLoadingPackage for %s: Package: %s (Async loading priority: %d)"), *NetRefHandleManager.PrintObjectFromNetRefHandle(RefHandle), *PackagePath.ToString(), AsyncLoadingPriority);
 
 	// Need timer?
 	FPendingAsyncLoadRequest LoadRequest(RefHandle, ReplicationSystem->GetElapsedTime());
@@ -2256,7 +2312,7 @@ bool FObjectReferenceCache::ShouldAsyncLoad() const
 		default: ensureMsgf( false, TEXT( "Invalid AsyncLoadMode: %i" ), (int32)AsyncLoadMode ); return false;
 	}
 }
-void FObjectReferenceCache::AddReferencedObjects(FReferenceCollector& ReferenceCollector)
+void FObjectReferenceCache::AddReferencedObjects(FReferenceCollector& ReferenceCollector, UReplicationSystem* ReplicationSystem)
 {
 	for (auto& It : QueuedBatchObjectReferences)
 	{
@@ -2303,11 +2359,12 @@ void FObjectReferenceCache::RemoveTrackedQueuedBatchObjectReference(const FNetRe
 	}
 }
 
-const TCHAR* FObjectReferenceCache::GetObjectRelativePath(FNetRefHandle NetRefHandle) const
+const TCHAR* FObjectReferenceCache::GetObjectRelativePath(FNetRefHandle NetRefHandle, UReplicationSystem* ReplicationSystem) const
 {
 	if (const FCachedNetObjectReference* CacheObject = ReferenceHandleToCachedReference.Find(NetRefHandle))
 	{
-		const TCHAR* ResolvedToken = StringTokenStore->ResolveToken(CacheObject->RelativePath);
+		FStringTokenStore* LocalStringTokenStore = ReplicationSystem->GetNetTokenStore()->GetDataStore<FStringTokenStore>();
+		const TCHAR* ResolvedToken = LocalStringTokenStore->ResolveToken(CacheObject->RelativePath);
 		return ResolvedToken;
 	}
 
diff --git a/Engine/Source/Runtime/Net/Iris/Private/Iris/ReplicationSystem/ObjectReferenceCache.h b/Engine/Source/Runtime/Net/Iris/Private/Iris/ReplicationSystem/ObjectReferenceCache.h
--- a/Engine/Source/Runtime/Net/Iris/Private/Iris/ReplicationSystem/ObjectReferenceCache.h
+++ b/Engine/Source/Runtime/Net/Iris/Private/Iris/ReplicationSystem/ObjectReferenceCache.h
@@ -40,7 +40,7 @@ public:
 
 	FObjectReferenceCache();
 
-	void Init(UReplicationSystem* ReplicationSystem);
+	void Init(bool bAuthority);
 
 	// Determine if the object is dynamic
 	bool IsDynamicObject(const UObject* Object) const;
@@ -49,10 +49,10 @@ public:
 	bool IsAuthority() const;
 
 	// Create and assign a new NetHandle to the object
-	FNetRefHandle CreateObjectReferenceHandle(const UObject* Object);
+	FNetRefHandle CreateObjectReferenceHandle(const UObject* Object, UReplicationSystem* ReplicationSystem);
 
 	// Create and assign a new pre-registered NetHandle to the object
-	FNetRefHandle PreRegisterObjectReferenceHandle(const UObject* Object);
+	FNetRefHandle PreRegisterObjectReferenceHandle(const UObject* Object, UReplicationSystem* ReplicationSystem);
 
 	// Get existing handle for object
 	FNetRefHandle GetObjectReferenceHandleFromObject(const UObject* Object, EGetRefHandleFlags GetRefHandleFlags = EGetRefHandleFlags::None) const;
@@ -79,19 +79,19 @@ public:
 	bool IsNetRefHandlePreRegistered(FNetRefHandle NetRefHandle) const;
 
 	// Find replicated outer
-	FNetObjectReference GetReplicatedOuter(const FNetObjectReference& Reference) const;
+	FNetObjectReference GetReplicatedOuter(const FNetObjectReference& Reference, UReplicationSystem* ReplicationSystem) const;
 
 	// Get or create a NetObjectReference from the object
-	FNetObjectReference GetOrCreateObjectReference(const UObject* Instance);
+	FNetObjectReference GetOrCreateObjectReference(const UObject* Instance, UReplicationSystem* ReplicationSystem);
 
 	// Get or create a NetObjectReference from the object identifed by path relative to outer
-	FNetObjectReference GetOrCreateObjectReference(const FString& ObjectPath, const UObject* Outer);
+	FNetObjectReference GetOrCreateObjectReference(const FString& ObjectPath, const UObject* Outer, UReplicationSystem* ReplicationSystem);
 
 	// Bind a nethandle and the object reference cache on the client
 	void BindRemoteReference(FNetRefHandle RefHandle, const UObject* Object);
 	
 	// Add reference for pre-registered object (client)
-	void AddPreRegisteredReference(FNetRefHandle RefHandle, const UObject* Object);
+	void AddPreRegisteredReference(FNetRefHandle RefHandle, const UObject* Object, UReplicationSystem* ReplicationSystem);
 
 	// Remove references to dynamic objects
 	void RemoveReference(FNetRefHandle RefHandle, const UObject* Object);
@@ -150,7 +150,7 @@ public:
 
 
 	// While async loading of pending must be mapped references we need to maintain references to already resolved objects as there will be no instance referencing them
-	void AddReferencedObjects(FReferenceCollector& ReferenceCollector);
+	void AddReferencedObjects(FReferenceCollector& ReferenceCollector, UReplicationSystem* ReplicationSystem);
 	void AddTrackedQueuedBatchObjectReference(const FNetRefHandle InHandle, const UObject* InObject);
 	void UpdateTrackedQueuedBatchObjectReference(const FNetRefHandle InHandle, const UObject* NewObject);
 	void RemoveTrackedQueuedBatchObjectReference(const FNetRefHandle InHandle);
@@ -158,7 +158,7 @@ public:
 	FString DescribeObjectReference(const FNetObjectReference Ref, const FNetObjectResolveContext& ResolveContext);
 
 	/** Return the stored relative path of a replicated object if it's in the cache. */
-	const TCHAR* GetObjectRelativePath(FNetRefHandle NetRefHandle) const;
+	const TCHAR* GetObjectRelativePath(FNetRefHandle NetRefHandle, UReplicationSystem* ReplicationSystem) const;
 
 private:
 	void AsyncPackageCallback(const FName& PackageName, UPackage* Package, EAsyncLoadingResult::Type Result);
@@ -204,7 +204,7 @@ private:
 		double RequestStartTime;
 	};
 
-	bool CreateObjectReferenceInternal(const UObject* Object, FNetObjectReference& OutReference);
+	bool CreateObjectReferenceInternal(const UObject* Object, FNetObjectReference& OutReference, UReplicationSystem* ReplicationSystem);
 
 
 	void ReadFullReferenceInternal(FNetSerializationContext& Context, FNetObjectReference& OutRef, uint32 RecursionCount);
@@ -215,7 +215,7 @@ private:
 	bool SupportsObjectInternal(const UObject* Object) const;
 	bool CanClientLoadObjectInternal(const UObject* Object, bool bIsDynamic) const;
 	bool ShouldIgnoreWhenMissing(FNetRefHandle RefHandle) const;
-	bool RenamePathForPie(uint32 ConnectionId, FString& Str, bool bReading);
+	bool RenamePathForPie(uint32 ConnectionId, FString& Str, bool bReading, UReplicationSystem* ReplicationSystem);
 
 	// Get the string path of RefHandle
 	FString FullPath(FNetRefHandle RefHandle, const FNetObjectResolveContext& ResolveContext) const;
@@ -232,7 +232,7 @@ private:
 	bool WriteMustBeMappedExports(FNetSerializationContext& Context, FInternalNetRefIndex ObjectIndex, TArrayView<const FNetObjectReference> ExportsView) const;
 	void ReadMustBeMappedExports(const FNetRefHandle& NetObjectHandle, FNetSerializationContext& Context, TArray<FNetRefHandle>* MustBeMappedExports, EIrisAsyncLoadingPriority& OutIrisAsyncLoadingPriority);
 
-	void StartAsyncLoadingPackage(FCachedNetObjectReference& Object, FName PackagePath, TAsyncLoadPriority AsyncLoadingPriority, const FNetRefHandle RefHandle, const bool bWasAlreadyAsyncLoading);
+	void StartAsyncLoadingPackage(FCachedNetObjectReference& Object, FName PackagePath, TAsyncLoadPriority AsyncLoadingPriority, const FNetRefHandle RefHandle, const bool bWasAlreadyAsyncLoading, UReplicationSystem* ReplicationSystem);
 	void ValidateAsyncLoadingPackage(FCachedNetObjectReference& Object, FName PackagePath, const FNetRefHandle RefHandle);
 
 private:
@@ -247,12 +247,6 @@ private:
 	// To properly clean up stale references referencing dynamic objects we need to track them
 	TMultiMap<FNetRefHandle, FNetRefHandle> HandleToDynamicOuter;
 	
-	UReplicationSystem* ReplicationSystem;
-	UObjectReplicationBridge* ReplicationBridge;
-	FNetTokenStore* NetTokenStore;
-	FStringTokenStore* StringTokenStore;
-	FNetRefHandleManager* NetRefHandleManager;
-	
 	// Do we have authority to create references?
 	uint32 bIsAuthority : 1;
 
diff --git a/Engine/Source/Runtime/Net/Iris/Private/Iris/ReplicationSystem/ObjectReferenceCacheFwd.h b/Engine/Source/Runtime/Net/Iris/Private/Iris/ReplicationSystem/ObjectReferenceCacheFwd.h
--- a/Engine/Source/Runtime/Net/Iris/Private/Iris/ReplicationSystem/ObjectReferenceCacheFwd.h
+++ b/Engine/Source/Runtime/Net/Iris/Private/Iris/ReplicationSystem/ObjectReferenceCacheFwd.h
@@ -32,5 +32,6 @@ namespace UE::Net
 		uint32 ConnectionId = InvalidConnectionId;
 		TAsyncLoadPriority AsyncLoadingPriority = INDEX_NONE; // default to -1 to trap codepaths that are missing the assignment
 		bool bForceSyncLoad = false;
+		UReplicationSystem* ReplicationSystem = nullptr;
 	};
 }
diff --git a/Engine/Source/Runtime/Net/Iris/Private/Iris/ReplicationSystem/ObjectReplicationBridge.cpp b/Engine/Source/Runtime/Net/Iris/Private/Iris/ReplicationSystem/ObjectReplicationBridge.cpp
--- a/Engine/Source/Runtime/Net/Iris/Private/Iris/ReplicationSystem/ObjectReplicationBridge.cpp
+++ b/Engine/Source/Runtime/Net/Iris/Private/Iris/ReplicationSystem/ObjectReplicationBridge.cpp
@@ -333,12 +333,12 @@ FString UObjectReplicationBridge::DescribeObjectReference(const UE::Net::FNetObj
 
 UE::Net::FNetObjectReference UObjectReplicationBridge::GetOrCreateObjectReference(const UObject* Instance) const
 {
-	return GetObjectReferenceCache()->GetOrCreateObjectReference(Instance);
+	return GetObjectReferenceCache()->GetOrCreateObjectReference(Instance, GetReplicationSystem());
 }
 
 UE::Net::FNetObjectReference UObjectReplicationBridge::GetOrCreateObjectReference(const FString& Path, const UObject* Outer) const
 {
-	return GetObjectReferenceCache()->GetOrCreateObjectReference(Path, Outer);
+	return GetObjectReferenceCache()->GetOrCreateObjectReference(Path, Outer, GetReplicationSystem());
 }
 
 void UObjectReplicationBridge::AddStaticDestructionInfo(const FString& ObjectPath, const UObject* Outer, const UE::Net::FDestructionParameters& Parameters)
@@ -559,12 +559,12 @@ bool UObjectReplicationBridge::IsNetRefHandlePreRegistered(FNetRefHandle Handle)
 
 UE::Net::FNetRefHandle UObjectReplicationBridge::PreRegisterNewObjectReferenceHandle(UObject* Instance)
 {
-	return ObjectReferenceCache->PreRegisterObjectReferenceHandle(Instance);
+	return ObjectReferenceCache->PreRegisterObjectReferenceHandle(Instance, GetReplicationSystem());
 }
 
 void UObjectReplicationBridge::PreRegisterObjectWithReferenceHandle(const UObject* Instance, FNetRefHandle Handle)
 {
-	ObjectReferenceCache->AddPreRegisteredReference(Handle, Instance);
+	ObjectReferenceCache->AddPreRegisteredReference(Handle, Instance, ReplicationSystem);
 }
 
 UE::Net::FNetRefHandle UObjectReplicationBridge::StartReplicatingNetObject(UObject* Instance, UE::Net::EReplicationFragmentTraits Traits, UE::Net::FNetObjectFactoryId NetFactoryId)
@@ -572,7 +572,7 @@ UE::Net::FNetRefHandle UObjectReplicationBridge::StartReplicatingNetObject(UObje
 	using namespace UE::Net;
 	using namespace UE::Net::Private;
 
-	FNetRefHandle AllocatedRefHandle = ObjectReferenceCache->CreateObjectReferenceHandle(Instance);
+	FNetRefHandle AllocatedRefHandle = ObjectReferenceCache->CreateObjectReferenceHandle(Instance, GetReplicationSystem());
 
 	// If we failed to assign a handle, or if the Handle already is replicating, just return the handle
 	if (!AllocatedRefHandle.IsValid())
@@ -1201,22 +1201,28 @@ void UObjectReplicationBridge::DetachInstanceFromRemote(FNetRefHandle Handle, ER
 		// Now ask the object's factory to handle the destruction
 		if (ensureMsgf(NetFactoryId != InvalidNetObjectFactoryId, TEXT("Received invalid netfactory for %s: %s"), *Handle.ToString(), *GetPathNameSafe(DestroyContext.DestroyedInstance)))
 		{
-			if (UNetObjectFactory* Factory = GetNetFactory(NetFactoryId))
+			if (DestroyReason != EReplicationBridgeDestroyInstanceReason::Migrated)
 			{
-PRAGMA_DISABLE_DEPRECATION_WARNINGS
-				// Call the deprecated function
-				if (DestroyReason != EReplicationBridgeDestroyInstanceReason::DoNotDestroy)
+				if (UNetObjectFactory* Factory = GetNetFactory(NetFactoryId))
 				{
-					Factory->DestroyReplicatedObject(DestroyContext);
-				}
+PRAGMA_DISABLE_DEPRECATION_WARNINGS
+					// Call the deprecated function
+					if (DestroyReason != EReplicationBridgeDestroyInstanceReason::DoNotDestroy)
+					{
+						Factory->DestroyReplicatedObject(DestroyContext);
+					}
 PRAGMA_ENABLE_DEPRECATION_WARNINGS
 
-				Factory->DetachedFromReplication(DestroyContext);
+					Factory->DetachedFromReplication(DestroyContext);
+				}
 			}
 		}
 	}
-	
-	UnregisterInstance(Handle);	
+
+	if (DestroyReason != EReplicationBridgeDestroyInstanceReason::Migrated)
+	{
+		UnregisterInstance(Handle);
+	}
 }
 
 void UObjectReplicationBridge::DetachSubObjectInstancesFromRemote(UE::Net::Private::FInternalNetRefIndex InternalObjectIndex, EReplicationBridgeDestroyInstanceReason DestroyReason, EReplicationBridgeDestroyInstanceFlags DestroyFlags)
@@ -1335,7 +1341,23 @@ FReplicationBridgeCreateNetRefHandleResult UObjectReplicationBridge::CreateNetRe
 	// Instantiate the replicated object
 	const UNetObjectFactory::FInstantiateContext InstantiateContext(WantedNetHandle, Context.SerializationContext.GetInternalContext()->ResolveContext, RootObjectOfSubObject);
 	
-	const UNetObjectFactory::FInstantiateResult Result = Factory->InstantiateReplicatedObjectFromHeader(InstantiateContext, Header.Get());
+	UNetObjectFactory::FInstantiateResult Result;
+
+	if (GetReplicationSystem()->GetProxyType() == UE::Net::EProxyType::Backend)
+	{
+		const FNetRefHandle CompleteLookupNetHandle = NetRefHandleManager->MakeNetRefHandle(WantedNetHandle.GetId(), GetReplicationSystem()->GetId());
+		if (UObject* ExistingObject = GetObjectReferenceCache()->GetObjectFromReferenceHandle(CompleteLookupNetHandle))
+		{
+			UE_LOG_BRIDGEID(LogIris, Verbose, TEXT("CreateNetRefHandleFromRemote: Reusing existing object found in the object reference cache: Object: %s NetHandle: %s"), *GetNameSafe(ExistingObject), *CompleteLookupNetHandle.ToString());
+
+			Result = { .Instance = ExistingObject };
+		}
+	}
+	
+	if (!Result.Instance)
+	{
+		Result = Factory->InstantiateReplicatedObjectFromHeader(InstantiateContext, Header.Get());
+	}
 
 	// Failed ?
 	if (!Result.Instance)
@@ -3099,7 +3121,7 @@ void UObjectReplicationBridge::OnErrorBlockedByAsyncLoading(FNetRefHandle RefHan
 	
 	for (FNetRefHandle MustBeMappedRef : MustBeMappedRefs)
 	{
-		UE_LOG_BRIDGEID(LogIris, Error, TEXT("AsyncLoading reference: %s | PathName: %s"), *NetRefHandleManager->PrintObjectFromNetRefHandle(MustBeMappedRef), GetObjectReferenceCache()->GetObjectRelativePath(MustBeMappedRef));
+		UE_LOG_BRIDGEID(LogIris, Error, TEXT("AsyncLoading reference: %s | PathName: %s"), *NetRefHandleManager->PrintObjectFromNetRefHandle(MustBeMappedRef), GetObjectReferenceCache()->GetObjectRelativePath(MustBeMappedRef, ReplicationSystem));
 	}
 		
 	ensureMsgf(false, TEXT("Blocked by AsyncLoading error: Replication blocked for %s due to %d unmapped references."), *NetRefHandleManager->PrintObjectFromNetRefHandle(RefHandle), MustBeMappedRefs.Num());
diff --git a/Engine/Source/Runtime/Net/Iris/Private/Iris/ReplicationSystem/ReplicationBridge.cpp b/Engine/Source/Runtime/Net/Iris/Private/Iris/ReplicationSystem/ReplicationBridge.cpp
--- a/Engine/Source/Runtime/Net/Iris/Private/Iris/ReplicationSystem/ReplicationBridge.cpp
+++ b/Engine/Source/Runtime/Net/Iris/Private/Iris/ReplicationSystem/ReplicationBridge.cpp
@@ -414,6 +414,11 @@ void UReplicationBridge::StopReplicatingNetRefHandle(FNetRefHandle Handle, EEndR
 			UE_LOG_REPLICATIONBRIDGE(Verbose, TEXT("Ignoring EndReplication called on object already PendingEndReplication %s."), *NetRefHandleManager->PrintObjectFromNetRefHandle(Handle));
 			return;
 		}
+		
+		if (EnumHasAnyFlags(EndReplicationFlags, EEndReplicationFlags::Migrated))
+		{
+			ObjectData.CloseReason = ECloseReplicationReason::Migrated;
+		}
 
 		if (EnumHasAnyFlags(EndReplicationFlags, EEndReplicationFlags::TearOff))
 		{
diff --git a/Engine/Source/Runtime/Net/Iris/Private/Iris/ReplicationSystem/ReplicationBridgeTypes.cpp b/Engine/Source/Runtime/Net/Iris/Private/Iris/ReplicationSystem/ReplicationBridgeTypes.cpp
--- a/Engine/Source/Runtime/Net/Iris/Private/Iris/ReplicationSystem/ReplicationBridgeTypes.cpp
+++ b/Engine/Source/Runtime/Net/Iris/Private/Iris/ReplicationSystem/ReplicationBridgeTypes.cpp
@@ -42,7 +42,12 @@ FString LexToString(EEndReplicationFlags EndReplicationFlags)
 		if (!Flags.IsEmpty()) { Flags += TEXT(','); };
 		Flags += TEXT("SkipPendingEndReplicationValidation");
 	}
-	
+	if (EnumHasAnyFlags(EndReplicationFlags, EEndReplicationFlags::Migrated))
+	{
+		if (!Flags.IsEmpty()) { Flags += TEXT(','); };
+		Flags += TEXT("Migrated");
+	}
+
 	return Flags;
 }
 
@@ -62,6 +67,10 @@ const TCHAR* LexToString(EReplicationBridgeDestroyInstanceReason Reason)
 	{
 		return TEXT("Destroy");
 	}
+	case EReplicationBridgeDestroyInstanceReason::Migrated:
+	{
+		return TEXT("Migrated");
+	}
 	default:
 	{
 		return TEXT("[Invalid]");
diff --git a/Engine/Source/Runtime/Net/Iris/Private/Iris/ReplicationSystem/ReplicationReader.cpp b/Engine/Source/Runtime/Net/Iris/Private/Iris/ReplicationSystem/ReplicationReader.cpp
--- a/Engine/Source/Runtime/Net/Iris/Private/Iris/ReplicationSystem/ReplicationReader.cpp
+++ b/Engine/Source/Runtime/Net/Iris/Private/Iris/ReplicationSystem/ReplicationReader.cpp
@@ -229,10 +229,29 @@ struct FReplicationReader::FDispatchObjectInfo
 	uint32 bProcessEndReplicationForHierarchy : 1 = false;
 	uint32 bShouldCallSubObjectCreatedFromReplication : 1 = false;
 	uint32 bDynamicObjectCreated : 1 = false;
+	ECloseReplicationReason CloseReason = ECloseReplicationReason::None;
 };
 
 //------------------------------------------------------------------------
 
+bool ReadCloseReason(FNetSerializationContext& Context, ECloseReplicationReason& OutReason)
+{
+	FNetBitStreamReader& Reader = *Context.GetBitStreamReader();
+
+	OutReason = ECloseReplicationReason::None;
+
+	const bool bHasDestroyReason = Reader.ReadBool();
+	if (bHasDestroyReason)
+	{
+		const uint32 MaxBits = static_cast<uint32>(ECloseReplicationReason::MaxBits);
+		OutReason = static_cast<ECloseReplicationReason>(Reader.ReadBits(MaxBits));
+	}
+
+	return !Context.HasErrorOrOverflow();
+}
+
+//------------------------------------------------------------------------
+
 // Helper class to deal with management of ObjectsToDispatch allocations from our temporary allocator
 class FReplicationReader::FObjectsToDispatchArray
 {
@@ -373,6 +392,7 @@ void FReplicationReader::Init(const FReplicationParameters& InParameters)
 	// Init resolve context
 	ResolveContext.ConnectionId = InParameters.ConnectionId;
 	ResolveContext.RemoteNetTokenStoreState = Parameters.ReplicationSystem->GetNetTokenStore()->GetRemoteNetTokenStoreState(InParameters.ConnectionId);
+	ResolveContext.ReplicationSystem = Parameters.ReplicationSystem;
 
 	// Find out if there's a PartialNetObjectAttachmentHandler so we can re-assemble split blobs
 	if (const UPartialNetObjectAttachmentHandler* Handler = ReplicationSystemInternal->GetNetBlobManager().GetPartialNetObjectAttachmentHandler())
@@ -490,11 +510,17 @@ uint32 FReplicationReader::ReadObjectsAndSubObjectsPendingDestroy(FNetSerializat
 			break;
 		}
 
+		ECloseReplicationReason CloseReason = ECloseReplicationReason::None;
+		if (!ReadCloseReason(Context, CloseReason))
+		{
+			break;
+		}
+
 		if (FPendingBatchData* PendingBatchData = PendingBatchHolder.Find(SubObjectRootOrHandle))
 		{
 			// We expect end replication requests for subobjects too so we should only handle one object at a time.
 			constexpr bool bShouldProcessHierarchy = false;
-			EnqueueEndReplication(PendingBatchData, bShouldDestroyInstance, bShouldProcessHierarchy, IncompleteHandle);
+			EnqueueEndReplication(PendingBatchData, bShouldDestroyInstance, bShouldProcessHierarchy, IncompleteHandle, CloseReason);
 			continue;
 		}
 
@@ -518,13 +544,14 @@ uint32 FReplicationReader::ReadObjectsAndSubObjectsPendingDestroy(FNetSerializat
 					Info.bHasState = 0U;
 					Info.bHasAttachments = 0U;
 					Info.bShouldCallSubObjectCreatedFromReplication = 0U;
+					Info.CloseReason = CloseReason;
 
 					// Mark for dispatch
 					ObjectsToDispatchArray->CommitPendingDispatchObjectInfo();
 				}
 				else
 				{
-					EndReplication(InternalIndex, false, bShouldDestroyInstance);
+					EndReplication(InternalIndex, false, bShouldDestroyInstance, CloseReason);
 				}
 			}
 			else
@@ -565,10 +592,16 @@ uint32 FReplicationReader::ReadRootObjectsPendingDestroy(FNetSerializationContex
 			break;
 		}
 
+		ECloseReplicationReason CloseReason = ECloseReplicationReason::None;
+		if (!ReadCloseReason(Context, CloseReason))
+		{
+			break;
+		}
+
 		if (FPendingBatchData* PendingBatchData = PendingBatchHolder.Find(IncompleteHandle))
 		{
 			constexpr bool bShouldProcessHierarchy = true;
-			EnqueueEndReplication(PendingBatchData, bShouldDestroyInstance, bShouldProcessHierarchy, IncompleteHandle);
+			EnqueueEndReplication(PendingBatchData, bShouldDestroyInstance, bShouldProcessHierarchy, IncompleteHandle, CloseReason);
 			continue;
 		}
 
@@ -593,6 +626,7 @@ uint32 FReplicationReader::ReadRootObjectsPendingDestroy(FNetSerializationContex
 					Info.bHasState = 0U;
 					Info.bHasAttachments = 0U;
 					Info.bShouldCallSubObjectCreatedFromReplication = 0U;
+					Info.CloseReason = CloseReason;
 
 					// Mark for dispatch
 					ObjectsToDispatchArray->CommitPendingDispatchObjectInfo();
@@ -600,7 +634,7 @@ uint32 FReplicationReader::ReadRootObjectsPendingDestroy(FNetSerializationContex
 				else
 				{
 					constexpr bool bTearOff = false;
-					EndReplication(InternalIndex, bTearOff, bShouldDestroyInstance);
+					EndReplication(InternalIndex, bTearOff, bShouldDestroyInstance, CloseReason);
 					for (FInternalNetRefIndex SubObjectIndex : NetRefHandleManager->GetSubObjects(InternalIndex))
 					{
 						// End replication for all subobjects and destroy dynamic ones
@@ -674,7 +708,7 @@ void FReplicationReader::CleanupObjectData(FReplicatedObjectInfo& ObjectInfo)
 	}
 }
 
-void FReplicationReader::EndReplication(FInternalNetRefIndex InternalIndex, bool bTearOff, bool bDestroyInstance)
+void FReplicationReader::EndReplication(FInternalNetRefIndex InternalIndex, bool bTearOff, bool bDestroyInstance, ECloseReplicationReason CloseReason)
 {
 	if (!ensure(InternalIndex != ObjectIndexForOOBAttachment))
 	{
@@ -692,6 +726,10 @@ void FReplicationReader::EndReplication(FInternalNetRefIndex InternalIndex, bool
 		{
 			DestroyReason = EReplicationBridgeDestroyInstanceReason::TearOff;
 		}
+		else if (CloseReason == ECloseReplicationReason::Migrated)
+		{
+			DestroyReason = EReplicationBridgeDestroyInstanceReason::Migrated;
+		}
 		else if (bDestroyInstance)
 		{
 			DestroyReason = EReplicationBridgeDestroyInstanceReason::Destroy;
@@ -2197,7 +2235,7 @@ void FReplicationReader::DispatchStateData(FNetSerializationContext& Context)
 				FDispatchObjectInfo& Info = *PostDispatchObjectInfo.Info;
 				if (Info.bReplicatedDestroyForSubObject)
 				{
-					EndReplication(Info.InternalIndex, Info.bTearOff, Info.bDestroy);
+					EndReplication(Info.InternalIndex, Info.bTearOff, Info.bDestroy, Info.CloseReason);
 				}
 			}
 		}
@@ -2600,7 +2638,7 @@ void FReplicationReader::DispatchEndReplication(FNetSerializationContext& Contex
 		if (Info.bDeferredEndReplication)
 		{
 			// Detach and destroy object
-			EndReplication(Info.InternalIndex, Info.bTearOff, Info.bDestroy);
+			EndReplication(Info.InternalIndex, Info.bTearOff, Info.bDestroy, Info.CloseReason);
 			if (Info.bProcessEndReplicationForHierarchy)
 			{
 				for (FInternalNetRefIndex SubObjectIndex : NetRefHandleManager->GetSubObjects(Info.InternalIndex))
@@ -2708,11 +2746,11 @@ void FReplicationReader::ProcessHugeObjectAttachment(FNetSerializationContext& C
 	}
 }
 
-bool FReplicationReader::EnqueueEndReplication(FPendingBatchData* PendingBatchData, bool bShouldDestroyInstance, bool bShouldProcessHierarchy, FNetRefHandle NetRefHandleToEndReplication)
+bool FReplicationReader::EnqueueEndReplication(FPendingBatchData* PendingBatchData, bool bShouldDestroyInstance, bool bShouldProcessHierarchy, FNetRefHandle NetRefHandleToEndReplication, ECloseReplicationReason CloseReason)
 {
 	UE_LOG(LogIris, Verbose, TEXT("FReplicationReader::EnqueueEndReplication for %s since %s has queued batches"), *NetRefHandleToEndReplication.ToString(), *PendingBatchData->Owner.ToString());
 		
-	const uint32 MaxNumDataBits = 66U;
+	const uint32 MaxNumDataBits = 66U + static_cast<uint32>(ECloseReplicationReason::MaxBits);
 	const uint32 NumDataWords = (MaxNumDataBits + 31U) / 32U;
 
 	// Enqueue BatchData
@@ -2732,7 +2770,17 @@ bool FReplicationReader::EnqueueEndReplication(FPendingBatchData* PendingBatchDa
 	// Write data to be parsed by ProcessQueuedBatches
 	WriteUint64(&Writer, NetRefHandleToEndReplication.GetId());
 	Writer.WriteBool(bShouldDestroyInstance);
+	
+	const bool bHasDestroyReason = CloseReason != ECloseReplicationReason::None;
+	Writer.WriteBool(bHasDestroyReason);
+	if (bHasDestroyReason)
+	{
+		const uint32 MaxBits = static_cast<uint32>(ECloseReplicationReason::MaxBits);
+		Writer.WriteBits(static_cast<uint32>(CloseReason), MaxBits);
+	}
+	
 	Writer.WriteBool(bShouldProcessHierarchy);
+
 	Writer.CommitWrites();
 	ensure(Writer.GetPosBits() <= MaxNumDataBits);
 	
@@ -2814,6 +2862,13 @@ void FReplicationReader::ProcessQueuedBatches()
 				// this can be the batch root or a subobject owned by the batched root.
 				const uint64 NetRefHandleIdToEndReplication = ReadUint64(&Reader);
 				const bool bShouldDestroyInstance = Reader.ReadBool();
+
+				ECloseReplicationReason CloseReason = ECloseReplicationReason::None;
+				if (!ReadCloseReason(Context, CloseReason))
+				{
+					break;
+				}
+
 				const bool bShouldProcessHierarchy = Reader.ReadBool();
 					
 				const FNetRefHandle NetRefHandleToEndReplication = FNetRefHandleManager::MakeNetRefHandleFromId(NetRefHandleIdToEndReplication);
@@ -2825,7 +2880,7 @@ void FReplicationReader::ProcessQueuedBatches()
 					constexpr bool bTearOff = false;
 					if (bShouldProcessHierarchy)
 					{
-						EndReplication(InternalIndex, bTearOff, bShouldDestroyInstance);
+						EndReplication(InternalIndex, bTearOff, bShouldDestroyInstance, CloseReason);
 						for (const FInternalNetRefIndex SubObjectIndex : NetRefHandleManager->GetSubObjects(InternalIndex))
 						{
 							// End replication for all subobjects and destroy dynamic ones
@@ -2836,7 +2891,7 @@ void FReplicationReader::ProcessQueuedBatches()
 					}
 					else
 					{
-						EndReplication(InternalIndex, bTearOff, bShouldDestroyInstance);
+						EndReplication(InternalIndex, bTearOff, bShouldDestroyInstance, CloseReason);
 					}
 				}
 
diff --git a/Engine/Source/Runtime/Net/Iris/Private/Iris/ReplicationSystem/ReplicationReader.h b/Engine/Source/Runtime/Net/Iris/Private/Iris/ReplicationSystem/ReplicationReader.h
--- a/Engine/Source/Runtime/Net/Iris/Private/Iris/ReplicationSystem/ReplicationReader.h
+++ b/Engine/Source/Runtime/Net/Iris/Private/Iris/ReplicationSystem/ReplicationReader.h
@@ -188,7 +188,7 @@ private:
 	FReplicatedObjectInfo& StartReplication(uint32 InternalIndex);
 
 	// Remove tracking info for the object with InternalIndex
-	void EndReplication(FInternalNetRefIndex InternalIndex, bool bTearOff, bool bDestroyInstance);
+	void EndReplication(FInternalNetRefIndex InternalIndex, bool bTearOff, bool bDestroyInstance, ECloseReplicationReason RemovalFlags = ECloseReplicationReason::None);
 
 	// Free any data allocated per object
 	void CleanupObjectData(FReplicatedObjectInfo& ObjectInfo);
@@ -226,7 +226,7 @@ private:
 
 	// If we are queuing data for a batch we must also defer calls to EndReplication
 	// This method writes this method in the form of a QueuedChunk
-	bool EnqueueEndReplication(FPendingBatchData* PendingBatchData, bool bShouldDestroyInstance, bool bShouldProcessHierarchy, FNetRefHandle NetRefHandleToEndReplication);
+	bool EnqueueEndReplication(FPendingBatchData* PendingBatchData, bool bShouldDestroyInstance, bool bShouldProcessHierarchy, FNetRefHandle NetRefHandleToEndReplication, ECloseReplicationReason CloseReason);
 
 	// Remove a handle from the hot and cold unresolved caches used by ResolveAndDispatchUnresolvedReferences(). If this handle is marked as unresolved
 	// again, it will be added to the hot cache.
diff --git a/Engine/Source/Runtime/Net/Iris/Private/Iris/ReplicationSystem/ReplicationSystem.cpp b/Engine/Source/Runtime/Net/Iris/Private/Iris/ReplicationSystem/ReplicationSystem.cpp
--- a/Engine/Source/Runtime/Net/Iris/Private/Iris/ReplicationSystem/ReplicationSystem.cpp
+++ b/Engine/Source/Runtime/Net/Iris/Private/Iris/ReplicationSystem/ReplicationSystem.cpp
@@ -24,6 +24,7 @@
 #include "Iris/ReplicationSystem/ReplicationSystemTypes.h"
 #include "Iris/ReplicationSystem/ReplicationTypes.h"
 #include "Iris/ReplicationSystem/ReplicationWriter.h"
+#include "Iris/ReplicationSystem/ObjectReferenceCache.h"
 
 #include "Iris/ReplicationSystem/Filtering/NetObjectFilter.h"
 #include "Iris/ReplicationSystem/Filtering/ReplicationFiltering.h"
@@ -104,8 +105,12 @@ public:
 			.bUseRemoteObjectReferences = Params.bUseRemoteObjectReferences,
 			.bAllowParallelTasks = Params.bAllowParallelTasks,
 			.bAllowMinimalUpdateIfNoConnections = Params.bAllowMinimalUpdateIfNoConnections,
+			.bUseSharedObjectReferenceCache = (Params.ProxyType == UE::Net::EProxyType::Backend)
 		}))
 	{
+		const bool bIsProxyBackend = (Params.ProxyType == UE::Net::EProxyType::Backend);
+		ensureMsgf(!bIsProxyBackend || (bIsProxyBackend && !InReplicationSystem->IsServer()), TEXT("A replication system configured as a proxy backend is expected to be a client connecting to a remote server."));
+
 		MetricNameTotalRootObjects = TEXT("TotalSubObjects");
 		MetricNameTotalSubObjects = TEXT("TotalRootObjects");
 		MetricNameAvgPendingObjectCount = TEXT("AvgPendingObjectCount");
@@ -160,6 +165,7 @@ public:
 			NetRefHandleManagerInitParams.InternalNetRefIndexInitSize = Params.InitialNetObjectListCount;
 			NetRefHandleManagerInitParams.InternalNetRefIndexGrowSize = Params.NetObjectListGrowCount;
 			NetRefHandleManagerInitParams.NetChunkedArrayCount = Params.PreAllocatedMemoryBuffersObjectCount;
+			NetRefHandleManagerInitParams.InitialNetRefHandleIndex = Params.InitialNetRefHandleIndex;
 			NetRefHandleManager.Init(NetRefHandleManagerInitParams);
 
 			NetRefHandleManager.GetOnMaxInternalNetRefIndexIncreasedDelegate().AddRaw(this, &FReplicationSystemImpl::OnMaxInternalNetRefIndexIncreased);
@@ -286,7 +292,7 @@ public:
 		// Init replication bridge
 		ReplicationSystemInternal.GetReplicationBridge()->Initialize(ReplicationSystem);
 
-		ReplicationSystemInternal.GetObjectReferenceCache().Init(ReplicationSystem);
+		ReplicationSystemInternal.GetObjectReferenceCache().Init(ReplicationSystem->IsServer());
 
 		// Init custom packagemap we use for capturing references for backwards compatible NetSerializers
 		{
@@ -834,6 +840,7 @@ void UReplicationSystem::Init(uint32 InId, const FReplicationSystemParams& Param
 	bAllowExecuteRPC = Params.bAllowExecuteRPC;
 	bAllowExecuteRepNotify = Params.bAllowExecuteRepNotifies;
 	bAlwaysSwapRolesOnReplication = Params.bAlwaysSwapRolesOnReplication;
+	ProxyType = Params.ProxyType;
 
 	ReplicationBridge = Params.ReplicationBridge;
 
@@ -2024,7 +2031,7 @@ void UReplicationSystem::AddReferencedObjects(UObject* InThis, FReferenceCollect
 	if (This->Impl.IsValid())
 	{
 		This->Impl->ReplicationSystemInternal.GetNetRefHandleManager().AddReferencedObjects(Collector);
-		This->Impl->ReplicationSystemInternal.GetObjectReferenceCache().AddReferencedObjects(Collector);
+		This->Impl->ReplicationSystemInternal.GetObjectReferenceCache().AddReferencedObjects(Collector, This);
 	}
 	Super::AddReferencedObjects(InThis, Collector);
 }
@@ -2131,6 +2138,10 @@ bool UReplicationSystem::IsUsingRemoteObjectReferences() const
 
 namespace UE::Net
 {
+	
+// FReplicationSystemFactory
+FReplicationSystemFactory::FReplicationSystemArray FReplicationSystemFactory::ReplicationSystems;
+TSharedPtr<UE::Net::Private::FObjectReferenceCache> FReplicationSystemFactory::SharedObjectReferenceCache;
 
 FReplicationSystemCreatedDelegate& FReplicationSystemFactory::GetReplicationSystemCreatedDelegate()
 {
@@ -2144,8 +2155,18 @@ FReplicationSystemDestroyedDelegate& FReplicationSystemFactory::GetReplicationSy
 	return Delegate;
 }
 
-// FReplicationSystemFactory
-FReplicationSystemFactory::FReplicationSystemArray FReplicationSystemFactory::ReplicationSystems;
+TSharedPtr<UE::Net::Private::FObjectReferenceCache> FReplicationSystemFactory::GetSharedObjectReferenceCache()
+{
+	if (!SharedObjectReferenceCache)
+	{
+		SharedObjectReferenceCache = MakeShared<UE::Net::Private::FObjectReferenceCache>();
+
+		// It's assumed that the shared object reference cache will be used by client replication systems.
+		SharedObjectReferenceCache->Init(false);
+	}
+
+	return SharedObjectReferenceCache;
+}
 
 UReplicationSystem* FReplicationSystemFactory::CreateReplicationSystem(const UReplicationSystem::FReplicationSystemParams& Params)
 {
@@ -2220,6 +2241,12 @@ void FReplicationSystemFactory::DestroyReplicationSystem(UReplicationSystem* Sys
 	System->Shutdown();
 	System->RemoveFromRoot();
 	System->MarkAsGarbage();
+
+	// If this is the last replication system, reset the shared object reference cache.
+	if (ReplicationSystems.IsEmpty())
+	{
+		SharedObjectReferenceCache.Reset();
+	}
 }
 
 TArrayView<UReplicationSystem*> FReplicationSystemFactory::GetAllReplicationSystems()
diff --git a/Engine/Source/Runtime/Net/Iris/Private/Iris/ReplicationSystem/ReplicationSystemInternal.h b/Engine/Source/Runtime/Net/Iris/Private/Iris/ReplicationSystem/ReplicationSystemInternal.h
--- a/Engine/Source/Runtime/Net/Iris/Private/Iris/ReplicationSystem/ReplicationSystemInternal.h
+++ b/Engine/Source/Runtime/Net/Iris/Private/Iris/ReplicationSystem/ReplicationSystemInternal.h
@@ -20,6 +20,7 @@
 #include "Iris/ReplicationSystem/ReplicationProtocolManager.h"
 #include "Iris/ReplicationSystem/NetBlob/NetBlobManager.h"
 #include "Net/Core/NetToken/NetToken.h"
+#include "Iris/ReplicationSystem/ReplicationSystem.h"
 #include "Iris/ReplicationSystem/StringTokenStore.h"
 #include "Iris/ReplicationSystem/NameTokenStore.h"
 #include "Iris/ReplicationSystem/WorldLocations.h"
@@ -40,6 +41,7 @@ struct FReplicationSystemInternalInitParams
 	bool bUseRemoteObjectReferences;
 	bool bAllowParallelTasks;
 	bool bAllowMinimalUpdateIfNoConnections;
+	bool bUseSharedObjectReferenceCache;
 };
 
 class FReplicationSystemInternal
@@ -53,7 +55,17 @@ public:
 	, IrisObjectReferencePackageMap(nullptr)
 	, Id(Params.ReplicationSystemId)
 	, bAllowParallelTasks(Params.bAllowParallelTasks)
-	{}
+	{
+		if (Params.bUseSharedObjectReferenceCache)
+		{
+			ObjectReferenceCache = FReplicationSystemFactory::GetSharedObjectReferenceCache();
+		}
+		
+		if (!ObjectReferenceCache)
+		{
+			ObjectReferenceCache = MakeShared<FObjectReferenceCache>();
+		}
+	}
 
 	FReplicationProtocolManager& GetReplicationProtocolManager() { return ReplicationProtocolManager; }
 
@@ -68,7 +80,7 @@ public:
 
 	FReplicationStateStorage& GetReplicationStateStorage() { return ReplicationStateStorage; }
 
-	FObjectReferenceCache& GetObjectReferenceCache() { return ObjectReferenceCache; }
+	FObjectReferenceCache& GetObjectReferenceCache() { return *ObjectReferenceCache; }
 
 	void SetReplicationBridge(UObjectReplicationBridge* InReplicationBridge) { ReplicationBridge = InReplicationBridge; }
 	UObjectReplicationBridge* GetReplicationBridge() const { return ReplicationBridge; }
@@ -187,7 +199,7 @@ private:
 	FNetObjectGroups Groups;
 	FReplicationConditionals Conditionals;
 	FReplicationPrioritization Prioritization;
-	FObjectReferenceCache ObjectReferenceCache;
+	TSharedPtr<FObjectReferenceCache> ObjectReferenceCache;
 	FNetBlobManager NetBlobManager;
 	FWorldLocations WorldLocations;
 	FDeltaCompressionBaselineManager DeltaCompressionBaselineManager;
diff --git a/Engine/Source/Runtime/Net/Iris/Private/Iris/ReplicationSystem/ReplicationWriter.cpp b/Engine/Source/Runtime/Net/Iris/Private/Iris/ReplicationSystem/ReplicationWriter.cpp
--- a/Engine/Source/Runtime/Net/Iris/Private/Iris/ReplicationSystem/ReplicationWriter.cpp
+++ b/Engine/Source/Runtime/Net/Iris/Private/Iris/ReplicationSystem/ReplicationWriter.cpp
@@ -2082,6 +2082,14 @@ uint32 FReplicationWriter::WriteRootObjectsPendingDestroy(FNetSerializationConte
 		const bool bShouldDestroyInstance = ObjectData.RefHandle.IsDynamic() || NetRefHandleManager->GetIsDestroyedStartupObject(InternalIndex);
 		Writer.WriteBool(bShouldDestroyInstance);
 
+		const bool bHasDestroyReason = ObjectData.CloseReason != ECloseReplicationReason::None;
+		Writer.WriteBool(bHasDestroyReason);
+		if (bHasDestroyReason)
+		{
+			const uint32 MaxBits = static_cast<uint32>(ECloseReplicationReason::MaxBits);
+			Writer.WriteBits(static_cast<uint32>(ObjectData.CloseReason), MaxBits);
+		}
+
 		if (!Writer.IsOverflown())
 		{
 			// Must update state before pushing record
diff --git a/Engine/Source/Runtime/Net/Iris/Private/Iris/Serialization/InternalNetSerializationContext.cpp b/Engine/Source/Runtime/Net/Iris/Private/Iris/Serialization/InternalNetSerializationContext.cpp
--- a/Engine/Source/Runtime/Net/Iris/Private/Iris/Serialization/InternalNetSerializationContext.cpp
+++ b/Engine/Source/Runtime/Net/Iris/Private/Iris/Serialization/InternalNetSerializationContext.cpp
@@ -40,6 +40,7 @@ void FInternalNetSerializationContext::Init(const FInitParameters& InitParams)
 	*this = FInternalNetSerializationContext(InitParams.ReplicationSystem);
 	PackageMap = InitParams.PackageMap;
 	ResolveContext = InitParams.ObjectResolveContext;
+	ResolveContext.ReplicationSystem = InitParams.ReplicationSystem;
 }
 
 }
diff --git a/Engine/Source/Runtime/Net/Iris/Private/Iris/Serialization/ObjectNetSerializer.cpp b/Engine/Source/Runtime/Net/Iris/Private/Iris/Serialization/ObjectNetSerializer.cpp
--- a/Engine/Source/Runtime/Net/Iris/Private/Iris/Serialization/ObjectNetSerializer.cpp
+++ b/Engine/Source/Runtime/Net/Iris/Private/Iris/Serialization/ObjectNetSerializer.cpp
@@ -335,7 +335,8 @@ void FObjectNetSerializerBase<T>::Quantize(FNetSerializationContext& Context, co
 	else
 	{
 		UObject* SourceObject = GetRawPtrValue(Source);
-		Target.SetNetReference(Context, Args, Context.IsInitializingDefaultState() ? FNetObjectReference() : InternalContext->ObjectReferenceCache->GetOrCreateObjectReference(SourceObject));
+		UReplicationSystem* ReplicationSystem = InternalContext->ReplicationSystem;
+		Target.SetNetReference(Context, Args, Context.IsInitializingDefaultState() ? FNetObjectReference() : InternalContext->ObjectReferenceCache->GetOrCreateObjectReference(SourceObject, ReplicationSystem));
 	}
 }
 
diff --git a/Engine/Source/Runtime/Net/Iris/Public/Iris/ReplicationSystem/ReplicationBridgeTypes.h b/Engine/Source/Runtime/Net/Iris/Public/Iris/ReplicationSystem/ReplicationBridgeTypes.h
--- a/Engine/Source/Runtime/Net/Iris/Public/Iris/ReplicationSystem/ReplicationBridgeTypes.h
+++ b/Engine/Source/Runtime/Net/Iris/Public/Iris/ReplicationSystem/ReplicationBridgeTypes.h
@@ -23,6 +23,8 @@ enum class EEndReplicationFlags : uint32
 	ClearNetPushId						= DestroyNetHandle << 1U,
 	/** Skip bPendingEndReplication Validation, In some cases we want to allow detaching instance from replicated object on clients, such as when shutting down */
 	SkipPendingEndReplicationValidation = ClearNetPushId << 1U,
+	/** An actor is migrated between servers. Only supported by clients with replication systems configured to use UE::Net::EProxyType::Backend. */
+	Migrated							= SkipPendingEndReplicationValidation << 1U,
 };
 ENUM_CLASS_FLAGS(EEndReplicationFlags);
 
@@ -57,6 +59,7 @@ enum class EReplicationBridgeDestroyInstanceReason : uint32
 	DoNotDestroy,
 	TearOff,
 	Destroy,
+	Migrated,
 };
 IRISCORE_API const TCHAR* LexToString(EReplicationBridgeDestroyInstanceReason Reason);
 
diff --git a/Engine/Source/Runtime/Net/Iris/Public/Iris/ReplicationSystem/ReplicationSystem.h b/Engine/Source/Runtime/Net/Iris/Public/Iris/ReplicationSystem/ReplicationSystem.h
--- a/Engine/Source/Runtime/Net/Iris/Public/Iris/ReplicationSystem/ReplicationSystem.h
+++ b/Engine/Source/Runtime/Net/Iris/Public/Iris/ReplicationSystem/ReplicationSystem.h
@@ -10,6 +10,7 @@
 #include "Iris/ReplicationSystem/NetRefHandle.h"
 #include "Iris/ReplicationSystem/NetObjectGroupHandle.h"
 #include "Iris/ReplicationSystem/ReplicationSystemTypes.h"
+#include "Net/Core/Connection/NetEnums.h"
 #include "Net/Core/Connection/ConnectionHandle.h"
 #include "Net/Core/NetHandle/NetHandle.h"
 
@@ -62,6 +63,7 @@ namespace UE::Net
 	{
 		class FReplicationSystemImpl;
 		class FReplicationSystemInternal;
+		class FObjectReferenceCache;
 	}
 }
 
@@ -159,6 +161,15 @@ public:
 
 		/** NetTokenStore */
 		UE::Net::FNetTokenStore* NetTokenStore = nullptr;
+		
+		/** 
+		 * The initial index to use when assigning new NetRefHandle.
+		 * Must be greater than 0, because 0 is reserved for an invalid handle.
+		 */
+		uint32 InitialNetRefHandleIndex = 1;
+
+		/** Configure the replication system as a proxy. */
+		UE::Net::EProxyType ProxyType = UE::Net::EProxyType::None;
 	};
 
 	/** @return The unique ID of the ReplicationSystem. */
@@ -174,6 +185,8 @@ public:
 	bool AllowExecuteRepNotify() const { return bAllowExecuteRepNotify; };
 	bool AlwaysSwapRolesOnReplication() const { return bAlwaysSwapRolesOnReplication; };
 
+	UE::Net::EProxyType GetProxyType() const { return ProxyType; };
+
 	/** @return Is this system configured to replicate object properties. */
 	bool AllowObjectReplication() { return bAllowObjectReplication; }
 
@@ -799,6 +812,8 @@ private:
 	uint32 bAllowExecuteRepNotify : 1;
 	uint32 bAlwaysSwapRolesOnReplication : 1;
 	uint32 bDoCollectGarbage : 1;
+
+	UE::Net::EProxyType ProxyType = UE::Net::EProxyType::None;
 };
 
 namespace UE::Net
@@ -834,10 +849,14 @@ public:
 	/** Static delegate that is triggered before we destroy a replication system. */
 	IRISCORE_API static FReplicationSystemDestroyedDelegate& GetReplicationSystemDestroyedDelegate();
 
+	static TSharedPtr<UE::Net::Private::FObjectReferenceCache> GetSharedObjectReferenceCache();
+
 private:
 
 	friend UReplicationSystem* GetReplicationSystem(uint32 Id);
 
+	static TSharedPtr<UE::Net::Private::FObjectReferenceCache> SharedObjectReferenceCache;
+
 	using FReplicationSystemArray = TArray<UReplicationSystem*, TInlineAllocator<UE_NUM_INLINE_REPLICATIONSYSTEMS>>;
 	IRISCORE_API static FReplicationSystemArray ReplicationSystems;
 };
