Iris
* Add missing API export to FNetReferenceCollector::Add

#rnx

diff --git a/Engine/Source/Runtime/Net/Iris/Public/Iris/Serialization/NetReferenceCollector.h b/Engine/Source/Runtime/Net/Iris/Public/Iris/Serialization/NetReferenceCollector.h
--- a/Engine/Source/Runtime/Net/Iris/Public/Iris/Serialization/NetReferenceCollector.h
+++ b/Engine/Source/Runtime/Net/Iris/Public/Iris/Serialization/NetReferenceCollector.h
@@ -37,7 +37,7 @@ public:
 	FNetReferenceCollector() : Traits(ENetReferenceCollectorTraits::None) {}
 	explicit FNetReferenceCollector(ENetReferenceCollectorTraits InTraits) : Traits(InTraits) {}
 	
-	void Add(const FNetReferenceInfo& ReferenceInfo, const FQuantizedObjectReference& Reference, const FNetSerializerChangeMaskParam& ChangeMaskInfo);
+	IRISCORE_API void Add(const FNetReferenceInfo& ReferenceInfo, const FQuantizedObjectReference& Reference, const FNetSerializerChangeMaskParam& ChangeMaskInfo);
 
 	const FReferenceInfoArray& GetCollectedReferences() const { return ReferenceInfos; }
 
