Iris - Implement FInstancedStructNetSerializer delta serialization support.

diff --git a/Engine/Plugins/Runtime/ReplicationSystemTestPlugin/Source/Private/Tests/Serialization/TestInstancedStructNetSerializer.cpp b/Engine/Plugins/Runtime/ReplicationSystemTestPlugin/Source/Private/Tests/Serialization/TestInstancedStructNetSerializer.cpp
--- a/Engine/Plugins/Runtime/ReplicationSystemTestPlugin/Source/Private/Tests/Serialization/TestInstancedStructNetSerializer.cpp
+++ b/Engine/Plugins/Runtime/ReplicationSystemTestPlugin/Source/Private/Tests/Serialization/TestInstancedStructNetSerializer.cpp
@@ -258,7 +258,7 @@ UE_NET_TEST_FIXTURE(FTestInstancedStructNetSerializerFixture, TestDequantizeSeri
 	UE_NET_ASSERT_EQ(InstancedStruct0, InstancedStruct1);
 }
 
-UE_NET_TEST_FIXTURE(FTestInstancedStructNetSerializerFixture, TestSerializeDeltaEqualStates)
+UE_NET_TEST_FIXTURE(FTestInstancedStructNetSerializerFixture, TestSerializeDeltaEqualUnsetInstances)
 {
 	QuantizeTwoStates();
 	SerializeDelta();
@@ -266,12 +266,31 @@ UE_NET_TEST_FIXTURE(FTestInstancedStructNetSerializerFixture, TestSerializeDelta
 
 UE_NET_TEST_FIXTURE(FTestInstancedStructNetSerializerFixture, TestSerializeDeltaNonEqualStates)
 {
-	SetNonDefaultInstanceState();
+	// Equal types but different states
+	InstancedStruct0.InitializeAs<FStructForInstancedStructTestA>();
+	InstancedStruct0.GetMutable<FStructForInstancedStructTestA>().SomeUint16 += 4711U;
+
+	InstancedStruct1.InitializeAs<FStructForInstancedStructTestA>();
+	InstancedStruct1.GetMutable<FStructForInstancedStructTestA>().SomeUint16 += 711U;
+
 	QuantizeTwoStates();
 	SerializeDelta();
 }
 
-UE_NET_TEST_FIXTURE(FTestInstancedStructNetSerializerFixture, TestDeserializeDeltaEqualStates)
+UE_NET_TEST_FIXTURE(FTestInstancedStructNetSerializerFixture, TestSerializeDeltaNonEqualTypes)
+{
+	// Different types
+	InstancedStruct0.InitializeAs<FStructForInstancedStructTestWithArray>();
+	InstancedStruct0.GetMutable<FStructForInstancedStructTestWithArray>().ArrayOfTestB.SetNum(17);
+
+	InstancedStruct1.InitializeAs<FStructForInstancedStructTestA>();
+	InstancedStruct1.GetMutable<FStructForInstancedStructTestA>().SomeUint16 += 711U;
+
+	QuantizeTwoStates();
+	SerializeDelta();
+}
+
+UE_NET_TEST_FIXTURE(FTestInstancedStructNetSerializerFixture, TestDeserializeDeltaEqualUnsetInstances)
 {
 	QuantizeTwoStates();
 	SerializeDelta();
@@ -280,15 +299,68 @@ UE_NET_TEST_FIXTURE(FTestInstancedStructNetSerializerFixture, TestDeserializeDel
 
 UE_NET_TEST_FIXTURE(FTestInstancedStructNetSerializerFixture, TestDeserializeDeltaNonEqualStates)
 {
-	SetNonDefaultInstanceState();
+	// Equal types but different states
+	InstancedStruct0.InitializeAs<FStructForInstancedStructTestA>();
+	InstancedStruct0.GetMutable<FStructForInstancedStructTestA>().SomeUint16 += 4711U;
+
+	InstancedStruct1.InitializeAs<FStructForInstancedStructTestA>();
+	InstancedStruct1.GetMutable<FStructForInstancedStructTestA>().SomeUint16 += 711U;
+
 	QuantizeTwoStates();
 	SerializeDelta();
 	DeserializeDelta();
 }
 
-UE_NET_TEST_FIXTURE(FTestInstancedStructNetSerializerFixture, TestDequantizeDeltaSerializedState)
+UE_NET_TEST_FIXTURE(FTestInstancedStructNetSerializerFixture, TestDeserializeDeltaNonEqualTypes)
 {
-	SetNonDefaultInstanceState();
+	// Different types
+	InstancedStruct0.InitializeAs<FStructForInstancedStructTestWithArray>();
+	InstancedStruct0.GetMutable<FStructForInstancedStructTestWithArray>().ArrayOfTestB.SetNum(17);
+
+	InstancedStruct1.InitializeAs<FStructForInstancedStructTestA>();
+	InstancedStruct1.GetMutable<FStructForInstancedStructTestA>().SomeUint16 += 711U;
+
+	QuantizeTwoStates();
+	SerializeDelta();
+	DeserializeDelta();
+}
+
+UE_NET_TEST_FIXTURE(FTestInstancedStructNetSerializerFixture, TestDequantizeDeltaSerializedUnsetInstances)
+{
+	QuantizeTwoStates();
+	SerializeDelta();
+	DeserializeDelta();
+	Dequantize();
+
+	UE_NET_ASSERT_EQ(InstancedStruct0, InstancedStruct1);
+}
+
+UE_NET_TEST_FIXTURE(FTestInstancedStructNetSerializerFixture, TestDequantizeDeltaSerializedNonEqualStates)
+{
+	// Equal types but different states
+	InstancedStruct0.InitializeAs<FStructForInstancedStructTestA>();
+	InstancedStruct0.GetMutable<FStructForInstancedStructTestA>().SomeUint16 += 4711U;
+
+	InstancedStruct1.InitializeAs<FStructForInstancedStructTestA>();
+	InstancedStruct1.GetMutable<FStructForInstancedStructTestA>().SomeUint16 += 711U;
+
+	QuantizeTwoStates();
+	SerializeDelta();
+	DeserializeDelta();
+	Dequantize();
+
+	UE_NET_ASSERT_EQ(InstancedStruct0, InstancedStruct1);
+}
+
+UE_NET_TEST_FIXTURE(FTestInstancedStructNetSerializerFixture, TestDequantizeDeltaSerializedNonEqualTypes)
+{
+	// Different types
+	InstancedStruct0.InitializeAs<FStructForInstancedStructTestWithArray>();
+	InstancedStruct0.GetMutable<FStructForInstancedStructTestWithArray>().ArrayOfTestB.SetNum(17);
+
+	InstancedStruct1.InitializeAs<FStructForInstancedStructTestA>();
+	InstancedStruct1.GetMutable<FStructForInstancedStructTestA>().SomeUint16 += 711U;
+
 	QuantizeTwoStates();
 	SerializeDelta();
 	DeserializeDelta();
diff --git a/Engine/Source/Runtime/Net/Iris/Private/Iris/Serialization/InstancedStructNetSerializer.cpp b/Engine/Source/Runtime/Net/Iris/Private/Iris/Serialization/InstancedStructNetSerializer.cpp
--- a/Engine/Source/Runtime/Net/Iris/Private/Iris/Serialization/InstancedStructNetSerializer.cpp
+++ b/Engine/Source/Runtime/Net/Iris/Private/Iris/Serialization/InstancedStructNetSerializer.cpp
@@ -252,14 +252,97 @@ void FInstancedStructNetSerializer::Deserialize(FNetSerializationContext& Contex
 
 void FInstancedStructNetSerializer::SerializeDelta(FNetSerializationContext& Context, const FNetSerializeDeltaArgs& Args)
 {
-	// Skip DC support for now. Need to figure out how to gracefully handle missing UScriptStruct on the receiving end.
-	Serialize(Context, Args);
+	const QuantizedType& Value = *reinterpret_cast<const QuantizedType*>(Args.Source);
+	const QuantizedType& PrevValue = *reinterpret_cast<const QuantizedType*>(Args.Prev);
+
+	FNetBitStreamWriter* Writer = Context.GetBitStreamWriter();
+
+	// If the type has changed proceed with regular serialization.
+	if (Writer->WriteBool(Value.StructName != PrevValue.StructName))
+	{
+		Serialize(Context, Args);
+		return;
+	}
+	
+	// If the instance isn't set there's nothing to delta serialize
+	if (Value.StructName.IsNone())
+	{
+		return;
+	}
+
+	FInstancedStructNetSerializerConfig* Config = const_cast<FInstancedStructNetSerializerConfig*>(static_cast<const FInstancedStructNetSerializerConfig*>(Args.NetSerializerConfig));
+
+	FStructNetSerializerConfig StructConfig;
+	StructConfig.StateDescriptor = Config->DescriptorCache.FindOrAddDescriptor(Value.StructName);
+	ensureMsgf(StructConfig.StateDescriptor.IsValid(), TEXT("Struct type is no longer resolvable: %s. Sending FInstancedStruct as uninitialized."), ToCStr(Value.StructName.ToString()));
+
+	if (Writer->WriteBool(StructConfig.StateDescriptor.IsValid()))
+	{
+		// Delta serialize struct data
+		FNetSerializeDeltaArgs DeltaSerializeArgs = Args;
+		DeltaSerializeArgs.Source = NetSerializerValuePointer(Value.StructData.GetData());
+		DeltaSerializeArgs.Prev = NetSerializerValuePointer(PrevValue.StructData.GetData());
+		DeltaSerializeArgs.NetSerializerConfig = NetSerializerConfigParam(&StructConfig);
+		StructNetSerializer->SerializeDelta(Context, DeltaSerializeArgs);
+	}
 }
 
 void FInstancedStructNetSerializer::DeserializeDelta(FNetSerializationContext& Context, const FNetDeserializeDeltaArgs& Args)
 {
-	// Skip DC support for now. Need to figure out how to gracefully handle missing UScriptStruct on the receiving end.
-	Deserialize(Context, Args);
+	FNetBitStreamReader* Reader = Context.GetBitStreamReader();
+	if (const bool bTypeDiffers = Reader->ReadBool())
+	{
+		Deserialize(Context, Args);
+		return;
+	}
+
+	QuantizedType& Value = *reinterpret_cast<QuantizedType*>(Args.Target);
+	const QuantizedType& PrevValue = *reinterpret_cast<const QuantizedType*>(Args.Prev);
+
+	FInstancedStructNetSerializerConfig* Config = const_cast<FInstancedStructNetSerializerConfig*>(static_cast<const FInstancedStructNetSerializerConfig*>(Args.NetSerializerConfig));
+
+	// If the instance isn't set there's nothing to deserialize
+	if (PrevValue.StructName.IsNone())
+	{
+		Reset(Context, Config, Value);
+		return;
+	}
+
+	// The sender may not have succeeded in serializing the data due to struct type no longer existing.
+	if (!Reader->ReadBool())
+	{
+		Reset(Context, Config, Value);
+		return;
+	}
+	
+	// Delta deserialize
+	FStructNetSerializerConfig StructConfig;
+	StructConfig.StateDescriptor = Config->DescriptorCache.FindOrAddDescriptor(PrevValue.StructName);
+	if (!StructConfig.StateDescriptor.IsValid())
+	{
+		ensureMsgf(StructConfig.StateDescriptor.IsValid(), TEXT("Unable to create ReplicationStateDescriptor for struct %s."), ToCStr(PrevValue.StructName.ToString()));
+		Context.SetError(GNetError_InvalidValue);
+		return;
+	}
+
+	// Fix up struct info and allocations if the target value is changing types.
+	if (Value.StructName != PrevValue.StructName)
+	{
+		FreeStructInstance(Context, Config, Value);
+		Value.StructData.AdjustSize(Context, StructConfig.StateDescriptor->InternalSize, StructConfig.StateDescriptor->InternalAlignment);
+		Value.StructType = PrevValue.StructType;
+		Value.StructDescriptorTraits = PrevValue.StructDescriptorTraits;
+		Value.StructName = PrevValue.StructName;
+	}
+
+	// Delta deserialize struct data
+	{
+		FNetDeserializeDeltaArgs DeltaDeserializeArgs = Args;
+		DeltaDeserializeArgs.Target = NetSerializerValuePointer(Value.StructData.GetData());
+		DeltaDeserializeArgs.Prev = NetSerializerValuePointer(PrevValue.StructData.GetData());
+		DeltaDeserializeArgs.NetSerializerConfig = NetSerializerConfigParam(&StructConfig);
+		StructNetSerializer->DeserializeDelta(Context, DeltaDeserializeArgs);
+	}
 }
 
 void FInstancedStructNetSerializer::Quantize(FNetSerializationContext& Context, const FNetQuantizeArgs& Args)
@@ -278,7 +361,7 @@ void FInstancedStructNetSerializer::Quantize(FNetSerializationContext& Context,
 		{
 			const FString& PathNameString = Struct->GetPathName();
 			const FName PathName(PathNameString);
-			// If the struct type is the same as previous instance we don't need to free memory or adjust allocations.
+			// If the struct type is different from the previous instance we need adjust the memory allocation.
 			if (PathName != Target.StructName)
 			{
 				// We need to free the previous struct data prior to overwriting it. Doing this early.
