Additional fixes for the Iris shared object reference cache.

diff --git a/Engine/Source/Runtime/Engine/Classes/Engine/NetDriver.h b/Engine/Source/Runtime/Engine/Classes/Engine/NetDriver.h
--- a/Engine/Source/Runtime/Engine/Classes/Engine/NetDriver.h
+++ b/Engine/Source/Runtime/Engine/Classes/Engine/NetDriver.h
@@ -20,7 +20,6 @@
 #include "UObject/CoreNet.h"
 #include "Channel.h"
 #include "Net/Core/Connection/ConnectionHandle.h"
-#include "Net/Core/Connection/NetEnums.h"
 #include "Net/Core/Misc/DDoSDetection.h"
 #include "IPAddress.h"
 #include "Net/NetAnalyticsTypes.h"
@@ -1506,7 +1505,7 @@ private:
 	bool bSwapRolesOfAllActors = false;
 	bool bReplicationOfActorsWithAnyRole = false;
 
-	UE::Net::EProxyType ProxyType = UE::Net::EProxyType::None;
+	UE::Net::EProxyType ProxyType;
 
 	/** A set of replicators that have queued RPCs for transactions (server only) */
 	TSet< FObjectReplicator* > ReplicatorsWithPendingTransactions;
diff --git a/Engine/Source/Runtime/Engine/Private/NetDriver.cpp b/Engine/Source/Runtime/Engine/Private/NetDriver.cpp
--- a/Engine/Source/Runtime/Engine/Private/NetDriver.cpp
+++ b/Engine/Source/Runtime/Engine/Private/NetDriver.cpp
@@ -700,6 +700,7 @@ UNetDriver::UNetDriver(const FObjectInitializer& ObjectInitializer)
 ,	ProcessQueuedBunchesCurrentFrameMilliseconds(0.0f)
 ,	DDoS()
 ,	LocalAddr(nullptr)
+,	ProxyType(UE::Net::EProxyType::None)
 ,	DummyNetworkObjectInfo(MakePimpl<FNetworkObjectInfo>())
 ,	NetworkObjects(new FNetworkObjectList)
 ,	LagState(ENetworkLagState::NotLagging)
diff --git a/Engine/Source/Runtime/Net/Iris/Private/Iris/ReplicationSystem/ChunkedDataStream/ChunkedDataReader.cpp b/Engine/Source/Runtime/Net/Iris/Private/Iris/ReplicationSystem/ChunkedDataStream/ChunkedDataReader.cpp
--- a/Engine/Source/Runtime/Net/Iris/Private/Iris/ReplicationSystem/ChunkedDataStream/ChunkedDataReader.cpp
+++ b/Engine/Source/Runtime/Net/Iris/Private/Iris/ReplicationSystem/ChunkedDataStream/ChunkedDataReader.cpp
@@ -77,6 +77,7 @@ FChunkedDataReader::FChunkedDataReader(const UDataStream::FInitParameters& InPar
 	// Setup internal context
 	ResolveContext.ConnectionId = InitParams.ConnectionId;
 	ResolveContext.RemoteNetTokenStoreState = NetTokenStore->GetRemoteNetTokenStoreState(InitParams.ConnectionId);
+	ResolveContext.ReplicationSystem = ReplicationSystem;
 
 	NetTokenResolveContext.NetTokenStore = NetTokenStore;
 	NetTokenResolveContext.RemoteNetTokenStoreState = ResolveContext.RemoteNetTokenStoreState;
diff --git a/Engine/Source/Runtime/Net/Iris/Private/Iris/ReplicationSystem/ReplicationReader.cpp b/Engine/Source/Runtime/Net/Iris/Private/Iris/ReplicationSystem/ReplicationReader.cpp
--- a/Engine/Source/Runtime/Net/Iris/Private/Iris/ReplicationSystem/ReplicationReader.cpp
+++ b/Engine/Source/Runtime/Net/Iris/Private/Iris/ReplicationSystem/ReplicationReader.cpp
@@ -2750,7 +2750,7 @@ bool FReplicationReader::EnqueueEndReplication(FPendingBatchData* PendingBatchDa
 {
 	UE_LOG(LogIris, Verbose, TEXT("FReplicationReader::EnqueueEndReplication for %s since %s has queued batches"), *NetRefHandleToEndReplication.ToString(), *PendingBatchData->Owner.ToString());
 		
-	const uint32 MaxNumDataBits = 66U + static_cast<uint32>(ECloseReplicationReason::MaxBits);
+	const uint32 MaxNumDataBits = 67U + static_cast<uint32>(ECloseReplicationReason::MaxBits);
 	const uint32 NumDataWords = (MaxNumDataBits + 31U) / 32U;
 
 	// Enqueue BatchData
