Iris
- Fixed bug where unicast rpcs did not respect RPCSendPolicy
- Fixed bug where replicated objects overcomiiting reliable attachment send queue was not properly marked dirty when queue could accept send more data.

diff --git a/Engine/Plugins/Runtime/ReplicationSystemTestPlugin/Source/Private/Tests/ReplicationSystem/RPC/TestObjectRPC.cpp b/Engine/Plugins/Runtime/ReplicationSystemTestPlugin/Source/Private/Tests/ReplicationSystem/RPC/TestObjectRPC.cpp
--- a/Engine/Plugins/Runtime/ReplicationSystemTestPlugin/Source/Private/Tests/ReplicationSystem/RPC/TestObjectRPC.cpp
+++ b/Engine/Plugins/Runtime/ReplicationSystemTestPlugin/Source/Private/Tests/ReplicationSystem/RPC/TestObjectRPC.cpp
@@ -542,6 +542,67 @@ namespace UE::Net::Private
 		}
 	}
 
+	UE_NET_TEST_FIXTURE(FRPCTestFixture, TestOverCommitReliableRPCDirtiesObjectWhenMoreAttachmentsCanBeSent)
+	{
+		// Add a client
+		FReplicationSystemTestClient* Client = CreateClient();
+
+		// Spawn object on server
+		UTestReplicatedObjectWithRPC* ServerObject = Server->CreateObject<UTestReplicatedObjectWithRPC>();
+
+		ServerObject->bIsServerObject = true;
+		ServerObject->ReplicationSystem = Server->GetReplicationSystem();
+		Server->ReplicationSystem->SetOwningNetConnection(ServerObject->NetRefHandle, 0x01);
+
+		// Send and deliver packet
+		Server->NetUpdate();
+		Server->SendAndDeliverTo(Client, true);
+		Server->PostSendUpdate();
+
+		UTestReplicatedObjectWithRPC* ClientObject = Cast<UTestReplicatedObjectWithRPC>(Client->GetReplicationBridge()->GetReplicatedObject(ServerObject->NetRefHandle));
+		
+		// Verify that created server handle now also exists on client
+		UE_NET_ASSERT_TRUE(ClientObject != nullptr);
+
+		ClientObject->ReplicationSystem = Client->GetReplicationSystem();
+
+		// Issue enough RPCs to overshoot reliable send window
+		const int32 RPCCount = 10000;
+
+		for (int32 It = 0; It < RPCCount; ++It)
+		{
+			ServerObject->ClientRPC();
+		}
+
+		int32 PrevCallOrder = ClientObject->ClientRPCCallOrder;
+
+		// Send a bunch of packets to fill up Reliable send queue and still have rpcs in the PreQueue
+		int32 PacketFlightCount = 6;
+		for (int32 PacketIt = 0; PacketIt < PacketFlightCount; ++PacketIt)
+		{
+			Server->NetUpdate();
+			Server->SendTo(Client);
+			Server->PostSendUpdate();
+		}
+
+		for (int32 PacketIt = 0; PacketIt < PacketFlightCount; ++PacketIt)
+		{
+			// Deliver the packets in flight
+			Server->DeliverTo(Client, true);
+		}
+
+		// Verify that we delivered RPCs
+		UE_NET_ASSERT_GT(ClientObject->ClientRPCCallOrder, PrevCallOrder);
+		PrevCallOrder = ClientObject->ClientRPCCallOrder;
+
+		// Send and deliver more packets
+		Server->UpdateAndSend({Client}, true);
+
+		// Confirm that we can deliver more RPCs after getting packet notifications
+		UE_NET_ASSERT_GT(ClientObject->ClientRPCCallOrder, PrevCallOrder);
+		PrevCallOrder = ClientObject->ClientRPCCallOrder;
+	}
+
 #if UE_AUTORTFM
 	UE_NET_TEST_FIXTURE(FRPCTestFixture, AbortingRPCWithStompDectionDoesNotCrash)
 	{
diff --git a/Engine/Source/Runtime/Net/Iris/Private/Iris/ReplicationSystem/AttachmentReplication.cpp b/Engine/Source/Runtime/Net/Iris/Private/Iris/ReplicationSystem/AttachmentReplication.cpp
--- a/Engine/Source/Runtime/Net/Iris/Private/Iris/ReplicationSystem/AttachmentReplication.cpp
+++ b/Engine/Source/Runtime/Net/Iris/Private/Iris/ReplicationSystem/AttachmentReplication.cpp
@@ -123,12 +123,14 @@ public:
 		}
 	}
 
-	void ProcessPacketDeliveryStatus(EPacketDeliveryStatus Status, const FReliableNetBlobQueue::FReplicationRecord& Record)
+	void ProcessPacketDeliveryStatus(EPacketDeliveryStatus Status, const FReliableNetBlobQueue::FReplicationRecord& Record, bool& bOutMarkHasAttachmentsToSendDirty)
 	{
 		ReliableQueue.ProcessPacketDeliveryStatus(Status, Record);
 		if (Status == EPacketDeliveryStatus::Delivered && Record.IsValid())
 		{
+			const uint32 PreUnsentBlobCount = ReliableQueue.GetUnsentBlobCount();
 			PopulateQueueFromPreQueue();
+			bOutMarkHasAttachmentsToSendDirty = ReliableQueue.GetUnsentBlobCount() > PreUnsentBlobCount;
 		}
 	}
 
@@ -492,14 +494,15 @@ void FNetObjectAttachmentSendQueue::CommitReplicationRecord(const FNetObjectAtta
 	}
 }
 
-void FNetObjectAttachmentSendQueue::ProcessPacketDeliveryStatus(EPacketDeliveryStatus Status, const FNetObjectAttachmentSendQueue::FReliableReplicationRecord& Record)
+void FNetObjectAttachmentSendQueue::ProcessPacketDeliveryStatus(EPacketDeliveryStatus Status, const FNetObjectAttachmentSendQueue::FReliableReplicationRecord& Record, bool& bOutMarkHasAttachmentsToSendDirty)
 {
+	bOutMarkHasAttachmentsToSendDirty = false;
 	if (ReliableQueue == nullptr)
 	{
 		return;
 	}
 
-	ReliableQueue->ProcessPacketDeliveryStatus(Status, Record);
+	ReliableQueue->ProcessPacketDeliveryStatus(Status, Record, bOutMarkHasAttachmentsToSendDirty);
 }
 
 bool FNetObjectAttachmentsWriter::Enqueue(ENetObjectAttachmentType Type, uint32 ObjectIndex, TArrayView<const TRefCountPtr<FNetBlob>> Attachments)
@@ -654,15 +657,22 @@ void FNetObjectAttachmentsWriter::CommitReplicationRecord(ENetObjectAttachmentTy
 	Queue->CommitReplicationRecord(Record);
 }
 
-void FNetObjectAttachmentsWriter::ProcessPacketDeliveryStatus(EPacketDeliveryStatus Status, ENetObjectAttachmentType Type, uint32 ObjectIndex, const FReliableReplicationRecord& Record)
+void FNetObjectAttachmentsWriter::ProcessPacketDeliveryStatus(EPacketDeliveryStatus Status, ENetObjectAttachmentType Type, uint32 ObjectIndex, const FReliableReplicationRecord& Record, bool& bOutMarkHasAttachmentsToSendDirty)
 {
+	bOutMarkHasAttachmentsToSendDirty = false;
 	FNetObjectAttachmentSendQueue* Queue = GetQueue(Type, ObjectIndex);
 	if (Queue == nullptr)
 	{
 		return;
 	}
 
-	Queue->ProcessPacketDeliveryStatus(Status, Record);
+	Queue->ProcessPacketDeliveryStatus(Status, Record, bOutMarkHasAttachmentsToSendDirty);
+}
+
+void FNetObjectAttachmentsWriter::ProcessPacketDeliveryStatus(EPacketDeliveryStatus Status, ENetObjectAttachmentType Type, uint32 ObjectIndex, const FReliableReplicationRecord& Record)
+{
+	bool bDummyMarkHasAttachmentsToSendDirty = false;
+	ProcessPacketDeliveryStatus(Status, Type, ObjectIndex, Record, bDummyMarkHasAttachmentsToSendDirty);
 }
 
 FNetObjectAttachmentSendQueue* FNetObjectAttachmentsWriter::GetQueue(ENetObjectAttachmentType Type, uint32 ObjectIndex)
diff --git a/Engine/Source/Runtime/Net/Iris/Private/Iris/ReplicationSystem/AttachmentReplication.h b/Engine/Source/Runtime/Net/Iris/Private/Iris/ReplicationSystem/AttachmentReplication.h
--- a/Engine/Source/Runtime/Net/Iris/Private/Iris/ReplicationSystem/AttachmentReplication.h
+++ b/Engine/Source/Runtime/Net/Iris/Private/Iris/ReplicationSystem/AttachmentReplication.h
@@ -99,7 +99,7 @@ private:
 
 	void CommitReplicationRecord(const FCommitRecord& Record);
 
-	void ProcessPacketDeliveryStatus(EPacketDeliveryStatus Status, const FReliableReplicationRecord& Record);
+	void ProcessPacketDeliveryStatus(EPacketDeliveryStatus Status, const FReliableReplicationRecord& Record, bool& bOutMarkHasAttachmentsToSendDirty);
 
 	FReliableSendQueue* ReliableQueue;
 	TResizableCircularQueue<TRefCountPtr<FNetBlob>> UnreliableQueue;
@@ -146,6 +146,7 @@ public:
 	void CommitReplicationRecord(ENetObjectAttachmentType Type, uint32 ObjectIndex, const FCommitRecord& Record);
 
 	void ProcessPacketDeliveryStatus(EPacketDeliveryStatus Status, ENetObjectAttachmentType Type, uint32 ObjectIndex, const FReliableReplicationRecord& Record);
+	void ProcessPacketDeliveryStatus(EPacketDeliveryStatus Status, ENetObjectAttachmentType Type, uint32 ObjectIndex, const FReliableReplicationRecord& Record, bool& bOutMarkHasAttachmentsToSendDirty);
 
 private:
 	bool NetBlobMightNeedSplitting(const TRefCountPtr<FNetObjectAttachment>& Attachment) const;
diff --git a/Engine/Source/Runtime/Net/Iris/Private/Iris/ReplicationSystem/ReplicationSystem.cpp b/Engine/Source/Runtime/Net/Iris/Private/Iris/ReplicationSystem/ReplicationSystem.cpp
--- a/Engine/Source/Runtime/Net/Iris/Private/Iris/ReplicationSystem/ReplicationSystem.cpp
+++ b/Engine/Source/Runtime/Net/Iris/Private/Iris/ReplicationSystem/ReplicationSystem.cpp
@@ -1384,12 +1384,22 @@ bool UReplicationSystem::SendRPC(const UObject* RootObject, const UObject* SubOb
 
 bool UReplicationSystem::SendRPC(uint32 ConnectionId, const UObject* RootObject, const UObject* SubObject, const UFunction* Function, const void* Parameters)
 {
+	using namespace UE::Net;
 	using namespace UE::Net::Private;
 	TRACE_CPUPROFILER_EVENT_SCOPE(SendRPC_SendUnicastRPC);
+
+	ENetObjectAttachmentSendPolicyFlags SendFlags = ENetObjectAttachmentSendPolicyFlags::None;
+	if (ReplicationSystemCVars::bAllowAttachmentSendPolicyFlags)
+	{
+		if (ENetObjectAttachmentSendPolicyFlags* Flags = Impl->AttachmentSendPolicyFlags.Find(FObjectKey(Function)))
+		{
+			SendFlags = *Flags;
+		}
+	}
 	
 	FNetBlobManager::FSendRPCContext RPCContext = { .RootObject = RootObject, .SubObject = SubObject, .Function = Function };
 	FNetBlobManager& NetBlobManager = Impl->ReplicationSystemInternal.GetNetBlobManager();
-	return NetBlobManager.SendUnicastRPC(ConnectionId, RPCContext, Parameters);
+	return NetBlobManager.SendUnicastRPC(ConnectionId, RPCContext, Parameters, SendFlags);
 }
 
 bool UReplicationSystem::SetRPCSendPolicyFlags(const UFunction* Function, UE::Net::ENetObjectAttachmentSendPolicyFlags SendFlags)
diff --git a/Engine/Source/Runtime/Net/Iris/Private/Iris/ReplicationSystem/ReplicationWriter.cpp b/Engine/Source/Runtime/Net/Iris/Private/Iris/ReplicationSystem/ReplicationWriter.cpp
--- a/Engine/Source/Runtime/Net/Iris/Private/Iris/ReplicationSystem/ReplicationWriter.cpp
+++ b/Engine/Source/Runtime/Net/Iris/Private/Iris/ReplicationSystem/ReplicationWriter.cpp
@@ -1339,9 +1339,40 @@ void FReplicationWriter::HandleDeliveredRecord(const FReplicationRecord::FRecord
 		break;
 	}
 
+	// If we now can send more data from our backlog, we need to mark object as dirty and having dirty attachments.
 	if (RecordInfo.HasAttachments)
 	{
-		Attachments.ProcessPacketDeliveryStatus(EPacketDeliveryStatus::Delivered, ENetObjectAttachmentType::Normal, InternalIndex, AttachmentRecord);
+		bool bMarkHasAttachmentDirty = false;
+
+		Attachments.ProcessPacketDeliveryStatus(EPacketDeliveryStatus::Delivered, ENetObjectAttachmentType::Normal, InternalIndex, AttachmentRecord, bMarkHasAttachmentDirty);
+
+		if (bMarkHasAttachmentDirty && CurrentState < EReplicatedObjectState::PendingDestroy)
+		{
+			// Mark object as having dirty changes
+			MarkObjectDirty(InternalIndex, "HandleDeliveredRecordMarkHasAttachmentDirty");
+			Info.HasAttachments = true;
+
+			// Bump prio if necessary
+			SchedulingPriorities[InternalIndex] = FPlatformMath::Max(FReplicationWriter::LostStatePriorityBump, SchedulingPriorities[InternalIndex]);
+
+			// Must also mark owner dirty
+			if (Info.IsSubObject)
+			{
+				FNetRefHandleManager::FReplicatedObjectData& ObjectData = NetRefHandleManager->GetReplicatedObjectDataNoCheck(InternalIndex);
+				if (ObjectData.IsOwnedSubObject())
+				{
+					FReplicationInfo& OwnerInfo = ReplicatedObjects[ObjectData.SubObjectRootIndex];
+					if (OwnerInfo.GetState() != EReplicatedObjectState::Invalid && OwnerInfo.GetState() < EReplicatedObjectState::PendingDestroy)
+					{
+						MarkObjectDirty(ObjectData.SubObjectRootIndex, "HandleDeliveredRecordMarkHasAttachmentDirty");
+						OwnerInfo.HasDirtySubObjects = 1U;
+
+						// Bump prio if necessary
+						SchedulingPriorities[ObjectData.SubObjectRootIndex] = FPlatformMath::Max(FReplicationWriter::LostStatePriorityBump, SchedulingPriorities[ObjectData.SubObjectRootIndex]);
+					}
+				}
+			}
+		}
 	}
 
 	// Must process WaitOnflush after attachments in order to correctly evaluate flush-status if needed
@@ -3359,7 +3390,7 @@ uint32 FReplicationWriter::WriteOOBAttachments(FNetSerializationContext& Context
 				return WrittenObjectCount;
 			}
 
-			WriteContext.bHasOOBAttachmentsToSend = Attachments.HasUnsentUnreliableAttachments(ENetObjectAttachmentType::OutOfBand, ObjectIndexForOOBAttachment);
+			WriteContext.bHasOOBAttachmentsToSend = Attachments.HasUnsentAttachments(ENetObjectAttachmentType::OutOfBand, ObjectIndexForOOBAttachment);
 			WrittenObjectCount += Result.NumWritten;
 		}
 	}
