Iris
- Fixed scheduling of dependent objects to not be reversed.
- Added support for scheduling creation dependencies before dependee to minimize queueing on client. (disabled)
- Additional logging in ReplicationWriter.

diff --git a/Engine/Plugins/Runtime/ReplicationSystemTestPlugin/Source/Private/Tests/ReplicationSystem/TestDependentObjects.cpp b/Engine/Plugins/Runtime/ReplicationSystemTestPlugin/Source/Private/Tests/ReplicationSystem/TestDependentObjects.cpp
--- a/Engine/Plugins/Runtime/ReplicationSystemTestPlugin/Source/Private/Tests/ReplicationSystem/TestDependentObjects.cpp
+++ b/Engine/Plugins/Runtime/ReplicationSystemTestPlugin/Source/Private/Tests/ReplicationSystem/TestDependentObjects.cpp
@@ -6,6 +6,7 @@
 #include "Iris/ReplicationSystem/Filtering/NetObjectFilterDefinitions.h"
 #include "Logging/LogScopedVerbosityOverride.h"
 #include "Tests/ReplicationSystem/Filtering/MockNetObjectFilter.h"
+#include "NetBlob/PartialNetBlobTestFixture.h"
 
 namespace UE::Net::Private
 {
@@ -947,6 +948,140 @@ UE_NET_TEST_FIXTURE(FTestFilteredDependentObjectFixture, TestDependentObjectSche
 	UE_NET_ASSERT_LT(ClientDependentObject->LastRepOrderCounter, ClientObject->LastRepOrderCounter);
 }
 
+UE_NET_TEST_FIXTURE(FReplicationSystemServerClientTestFixture, TestCreationDependencyCreationOrder)
+{
+	//LOG_SCOPE_VERBOSITY_OVERRIDE(LogIris, ELogVerbosity::VeryVerbose);
+	TGuardConsoleVariable<bool> CVarOverride(TEXT("net.Iris.ReplicationWriter.ScheduleCreationDependenciesFirst"), true);
+
+	UReplicationSystem* ReplicationSystem = Server->ReplicationSystem;
+	UReplicatedTestObjectBridge* Bridge = Server->GetReplicationBridge();
+
+	// Add a client
+	FReplicationSystemTestClient* Client = CreateClient();
+
+	// Spawn object on server
+	UReplicatedSubObjectOrderObject* ServerObject = Server->CreateObject<UReplicatedSubObjectOrderObject>();
+
+	// Spawn object with a creation dependency on ServerObject
+	UReplicatedSubObjectOrderObject* ServerDependentObject = Server->CreateObject<UReplicatedSubObjectOrderObject>();
+
+	Bridge->AddCreationDependencyLink(ServerObject->NetRefHandle, ServerDependentObject->NetRefHandle);
+
+	// Set low prio in parent
+	ReplicationSystem->SetStaticPriority(ServerObject->NetRefHandle, 0.1f);
+
+	// Set higher prio on dependent object to make it more important than parent, so that it will accrue scheduling priority faster
+	ReplicationSystem->SetStaticPriority(ServerDependentObject->NetRefHandle, 10.f);
+
+	// Set mini packet size to fail replication and accrue higher scheduling priorities
+	Server->SetMaxSendPacketSize(16);
+
+	// Send and deliver packet
+	Server->UpdateAndSend({Client});
+
+	// Verify that objects have not replicated
+	UReplicatedSubObjectOrderObject* ClientObject = Cast<UReplicatedSubObjectOrderObject>(Client->GetReplicationBridge()->GetReplicatedObject(ServerObject->NetRefHandle));
+	UReplicatedSubObjectOrderObject* ClientDependentObject = Cast<UReplicatedSubObjectOrderObject>(Client->GetReplicationBridge()->GetReplicatedObject(ServerDependentObject->NetRefHandle));
+
+	UE_NET_ASSERT_EQ(ClientObject, nullptr);
+	UE_NET_ASSERT_EQ(ClientDependentObject, nullptr);
+
+	// Restore packet size
+	Server->SetMaxSendPacketSize(2048);
+
+	// Reset RepOrderCounter
+	UReplicatedSubObjectOrderObject::RepOrderCounter = 0U;
+
+	// Send and deliver packet
+	Server->UpdateAndSend({Client});
+
+	// Verify that objects have replicated
+	ClientObject = Cast<UReplicatedSubObjectOrderObject>(Client->GetReplicationBridge()->GetReplicatedObject(ServerObject->NetRefHandle));
+	ClientDependentObject = Cast<UReplicatedSubObjectOrderObject>(Client->GetReplicationBridge()->GetReplicatedObject(ServerDependentObject->NetRefHandle));
+
+	UE_NET_ASSERT_NE(ClientObject, nullptr);
+	UE_NET_ASSERT_NE(ClientDependentObject, nullptr);
+
+	// Parent should have replicated before ClientDependentObject
+	UE_NET_ASSERT_GT(ClientDependentObject->LastRepOrderCounter, ClientObject->LastRepOrderCounter);
+}
+
+UE_NET_TEST_FIXTURE(FReplicationSystemServerClientTestFixture, TestNestedCreationDependenciesCreationOrder)
+{
+	//LOG_SCOPE_VERBOSITY_OVERRIDE(LogIris, ELogVerbosity::VeryVerbose);
+	TGuardConsoleVariable<bool> CVarOverride(TEXT("net.Iris.ReplicationWriter.ScheduleCreationDependenciesFirst"), true);
+
+	UReplicationSystem* ReplicationSystem = Server->ReplicationSystem;
+	UReplicatedTestObjectBridge* Bridge = Server->GetReplicationBridge();
+
+	// Add a client
+	FReplicationSystemTestClient* Client = CreateClient();
+
+	// Spawn object on server
+	UReplicatedSubObjectOrderObject* ServerObject = Server->CreateObject<UReplicatedSubObjectOrderObject>();
+	ServerObject->IntA = 0U;
+
+	TArray<UReplicatedSubObjectOrderObject*> Objects;
+	Objects.Add(ServerObject);
+
+	// Set low prio in parent
+	ReplicationSystem->SetStaticPriority(ServerObject->NetRefHandle, 0.1f);
+
+	// Spawn a bunch of dependent objects, setting prio
+	const int32 DependentObjectCount = 32;
+
+	for (uint32 ObjectIndex = 1U; ObjectIndex < DependentObjectCount; ++ObjectIndex)
+	{
+		UReplicatedSubObjectOrderObject* DependentObject = Server->CreateObject<UReplicatedSubObjectOrderObject>();
+		DependentObject->IntA = ObjectIndex;
+
+		// Increasing prio
+		float Prio = ObjectIndex * 1.0f;
+		ReplicationSystem->SetStaticPriority(DependentObject->NetRefHandle, Prio);
+
+		// Set up a nested creation dependency chain
+		Bridge->AddCreationDependencyLink(Objects.Last()->NetRefHandle, DependentObject->NetRefHandle);
+
+		Objects.Add(DependentObject);
+	}
+
+	// Set mini packet size to fail replication and accrue higher scheduling priorities
+	Server->SetMaxSendPacketSize(16);
+
+	// Send and deliver packet
+	Server->UpdateAndSend({Client});
+
+	// Restore packet size.
+	Server->SetMaxSendPacketSize(2048);
+
+	const int32 StartSpawnedObjectsCount =  Client->CreatedObjects.Num();
+
+	// Reset RepOrderCounter, important as this is a shared static.
+	UReplicatedSubObjectOrderObject::RepOrderCounter = 0U;
+	
+	// Send packets until objects are created
+	int MaxRetryCount = 32;
+	do 
+	{
+		Server->UpdateAndSend({Client});
+		--MaxRetryCount;
+	}
+	while ((Client->CreatedObjects.Num() < (DependentObjectCount + StartSpawnedObjectsCount)) && MaxRetryCount > 0);
+
+	// Verify replication order (even if it is not guaranteed if objects are scheduled independently, but for the context of this test it is)
+	UReplicatedSubObjectOrderObject* ClientObject = Cast<UReplicatedSubObjectOrderObject>(Client->GetReplicationBridge()->GetReplicatedObject(ServerObject->NetRefHandle));
+	UE_NET_ASSERT_NE(ClientObject, nullptr);
+	
+	TArray<UReplicatedSubObjectOrderObject*> ClientObjects;
+	for (UReplicatedSubObjectOrderObject* ServerDependentObject : Objects)
+	{
+		UReplicatedSubObjectOrderObject* ClientDependentObject = Cast<UReplicatedSubObjectOrderObject>(Client->GetReplicationBridge()->GetReplicatedObject(ServerDependentObject->NetRefHandle));
+		UE_NET_ASSERT_NE(ClientDependentObject, nullptr);
+		UE_NET_ASSERT_EQ(ClientDependentObject->LastRepOrderCounter, (uint32)ClientDependentObject->IntA + 1);
+		ClientObjects.Add(ClientDependentObject);
+	}
+}
+
 UE_NET_TEST_FIXTURE(FTestFilteredDependentObjectFixture, TestCircularDependencyIsDenied)
 {
 	UReplicatedTestObjectBridge* Bridge = Server->GetReplicationBridge();
diff --git a/Engine/Source/Runtime/Net/Iris/Private/Iris/IrisConfigInternal.h b/Engine/Source/Runtime/Net/Iris/Private/Iris/IrisConfigInternal.h
--- a/Engine/Source/Runtime/Net/Iris/Private/Iris/IrisConfigInternal.h
+++ b/Engine/Source/Runtime/Net/Iris/Private/Iris/IrisConfigInternal.h
@@ -18,7 +18,7 @@
 #if (UE_BUILD_SHIPPING || UE_BUILD_TEST)
 #	define UE_NET_ENABLE_REPLICATIONWRITER_LOG 0
 #else
-#	define UE_NET_ENABLE_REPLICATIONWRITER_LOG 0
+#	define UE_NET_ENABLE_REPLICATIONWRITER_LOG 1
 #endif
 #endif
 
diff --git a/Engine/Source/Runtime/Net/Iris/Private/Iris/ReplicationSystem/ReplicationWriter.cpp b/Engine/Source/Runtime/Net/Iris/Private/Iris/ReplicationSystem/ReplicationWriter.cpp
--- a/Engine/Source/Runtime/Net/Iris/Private/Iris/ReplicationSystem/ReplicationWriter.cpp
+++ b/Engine/Source/Runtime/Net/Iris/Private/Iris/ReplicationSystem/ReplicationWriter.cpp
@@ -44,15 +44,16 @@
 #include <cmath> // std::nextafter
 
 #if UE_NET_ENABLE_REPLICATIONWRITER_LOG
-#	define UE_LOG_REPLICATIONWRITER(Format, ...)  UE_LOG(LogIris, Log, Format, ##__VA_ARGS__)
-#	define UE_LOG_REPLICATIONWRITER_CONN(Format, ...)  UE_LOG(LogIris, Log, TEXT("Conn: %u ") Format, Parameters.ConnectionId, ##__VA_ARGS__)
+#	define UE_LOG_REPWRITER(Verbosity, Format, ...)  UE_LOG(LogIris, Verbosity, TEXT("Conn: %u ") Format, Parameters.ConnectionId, ##__VA_ARGS__)
 #else
-#	define UE_LOG_REPLICATIONWRITER(...)
-#	define UE_LOG_REPLICATIONWRITER_CONN(Format, ...)
+#	define UE_LOG_REPWRITER(...)
 #endif
 
-#define UE_LOG_REPLICATIONWRITER_WARNING(Format, ...)  UE_LOG(LogIris, Warning, Format, ##__VA_ARGS__)
-#define UE_CLOG_REPLICATIONWRITER_WARNING(Condition, Format, ...)  UE_CLOG(Condition, LogIris, Warning, Format, ##__VA_ARGS__)
+#define UE_LOG_REPWRITER_WARNING(Format, ...) UE_LOG(LogIris, Warning, TEXT("Conn: %u ") Format, Parameters.ConnectionId, ##__VA_ARGS__)
+#define UE_LOG_REPWRITER_ERROR(Format, ...) UE_LOG(LogIris, Error, TEXT("Conn: %u ") Format, Parameters.ConnectionId, ##__VA_ARGS__)
+#define UE_CLOG_REPLICATIONWRITER_WARNING(Condition, Format, ...)  UE_CLOG(Condition, LogIris, Warning, TEXT("Conn: %u ") Format, Parameters.ConnectionId, ##__VA_ARGS__)
+
+#define UE_NET_ENFORCE_CREATION_DEPENDENCIES_ON_SERVER 0
 
 namespace UE::Net::Private
 {
@@ -110,6 +111,10 @@ static FAutoConsoleVariableRef CVarReplicationWriterCannotSendWarningInterval(TE
 static bool bAllowSendStateWithEndReplicationForSubObjects = false;
 static FAutoConsoleVariableRef CvarAllowSendStateWithEndReplicationForSubObjects(TEXT("net.Iris.AllowSendStateWithEndReplicationForSubObjects"), bAllowSendStateWithEndReplicationForSubObjects, TEXT("Allow state to be sent together with subobject destroy. Default is false. Cannot be true if net.Iris.ImmediateDispatchEndReplicationForSubObjects is true."));
 
+// net.Iris.ReplicationWriter.ScheduleCreationDependenciesFirst
+static bool bScheduleCreationDependenciesFirst = false;
+static FAutoConsoleVariableRef CVarScheduleCreationDependenciesFirst(TEXT("net.Iris.ReplicationWriter.ScheduleCreationDependenciesFirst"), bScheduleCreationDependenciesFirst, TEXT("If enabled we will schedule an objects creation dependencies before the object if they are not confirmed as created."));
+
 // Drop unsent ordered unreliable attachments at the end of the tick, 
 // -1 = Allow them to be queued for next tick, 
 // 0  = Always drop at the end of the tick
@@ -349,7 +354,7 @@ void FReplicationWriter::SetState(uint32 InternalIndex, EReplicatedObjectState N
 	FReplicationInfo& Info = GetReplicationInfo(InternalIndex);
 
 	const EReplicatedObjectState CurrentState = Info.GetState();
-	UE_LOG_REPLICATIONWRITER_CONN(TEXT("ReplicationWriter.SetState for ( InternalIndex: %u ) %s -> %s"), InternalIndex, LexToString(CurrentState), LexToString(NewState));
+	UE_LOG_REPWRITER(Verbose, TEXT("ReplicationWriter.SetState for ( InternalIndex: %u ) %s -> %s"), InternalIndex, LexToString(CurrentState), LexToString(NewState));
 
 	switch (NewState)
 	{
@@ -568,7 +573,7 @@ void FReplicationWriter::StartReplication(uint32 InternalIndex)
 	// Currently we also do this for dependent objects to support objects with zero priority that should only replicate with parents
 	SchedulingPriorities[InternalIndex] = (Data.IsDependentObject() || Info.IsSubObject) ? 0.f : CreatePriority;
 
-	UE_LOG_REPLICATIONWRITER_CONN(TEXT("ReplicationWriter.StartReplication for ( InternalIndex: %u ) %s"), InternalIndex, ToCStr(Data.RefHandle.ToString()));
+	UE_LOG_REPWRITER(Verbose, TEXT("ReplicationWriter.StartReplication for ( InternalIndex: %u ) %s"), InternalIndex, ToCStr(Data.RefHandle.ToString()));
 
 	ObjectsWithDirtyChanges.SetBit(InternalIndex);
 
@@ -610,7 +615,7 @@ void FReplicationWriter::StopReplication(uint32 InternalIndex)
 	// Explicitly remove from objects in scope since we might call StopReplication from outside ScopeUpdate
 	ObjectsInScope.ClearBit(InternalIndex);
 
-	UE_LOG_REPLICATIONWRITER_CONN(TEXT("ReplicationWriter.StopReplication for ( InternalIndex: %u )"), InternalIndex);
+	UE_LOG_REPWRITER(Verbose, TEXT("ReplicationWriter.StopReplication for ( InternalIndex: %u )"), InternalIndex);
 	NetRefHandleManager->ReleaseNetObjectRef(InternalIndex);
 
 	Attachments.DropAllAttachments(ENetObjectAttachmentType::Normal, InternalIndex);
@@ -852,7 +857,7 @@ void FReplicationWriter::UpdateScope(const FNetBitArrayView& UpdatedScope)
 		}
 		else
 		{
-			UE_LOG_REPLICATIONWRITER_CONN(TEXT("New object added to scope, Waiting to start replication for ( InternalIndex: %u ) currently in State: %s "), Index, LexToString(State));
+			UE_LOG_REPWRITER(Verbose, TEXT("New object added to scope, Waiting to start replication for ( InternalIndex: %u ) currently in State: %s "), Index, LexToString(State));
 
 			ensureMsgf(!ObjectsWithDirtyChanges.GetBit(Index) , TEXT("New object added to scope, Waiting to start replication for ( InternalIndex: %u ) currently in State: %s "), Index, LexToString(State));
 			ensureMsgf(!Info.HasDirtyChangeMask, TEXT("New object added to scope, Waiting to start replication for ( InternalIndex: %u ) currently in State: %s "), Index, LexToString(State));
@@ -968,7 +973,7 @@ void FReplicationWriter::UpdateDirtyGlobalLifetimeConditionals(TArrayView<FInter
 
 					if (ObjectsInScope.GetBit(RootInternalObjectIndex) && (OwnerInfo.GetState() != EReplicatedObjectState::Invalid && OwnerInfo.GetState() < EReplicatedObjectState::PendingDestroy))
 					{
-						UE_LOG_REPLICATIONWRITER_CONN(TEXT("UpdateDirtyGlobalLifetimeConditionals for - %s"), ToCStr(NetRefHandleManager->PrintObjectFromIndex(RootInternalObjectIndex)));
+						UE_LOG_REPWRITER(Verbose, TEXT("UpdateDirtyGlobalLifetimeConditionals for - %s"), ToCStr(NetRefHandleManager->PrintObjectFromIndex(RootInternalObjectIndex)));
 						
 						MarkObjectDirty(RootInternalObjectIndex, "UpdateDirtyGlobalLifetimeConditionals");
 						OwnerInfo.HasDirtyConditionals = 1U;
@@ -978,7 +983,7 @@ void FReplicationWriter::UpdateDirtyGlobalLifetimeConditionals(TArrayView<FInter
 			}
 			else
 			{
-				UE_LOG_REPLICATIONWRITER_CONN(TEXT("UpdateDirtyGlobalLifetimeConditionals for - %s"), ToCStr(NetRefHandleManager->PrintObjectFromIndex(InternalObjectIndex)));
+				UE_LOG_REPWRITER(Verbose, TEXT("UpdateDirtyGlobalLifetimeConditionals for - %s"), ToCStr(NetRefHandleManager->PrintObjectFromIndex(InternalObjectIndex)));
 	
 				MarkObjectDirty(InternalObjectIndex, "UpdateDirtyGlobalLifetimeConditionals2");
 				Info.HasDirtyConditionals = 1U;
@@ -1072,7 +1077,7 @@ void FReplicationWriter::InternalUpdateDirtyChangeMasks(const FChangeMaskCache&
 		}
 	}
 
-	//UE_LOG_REPLICATIONWRITER(TEXT("FReplicationWriter::UpdateDirtyChangeMasks() Updated %u Objects for ConnectionId:%u, ReplicationSystemId: %u."), CachedChangeMasks.Indices.Num(), Parameters.ConnectionId, Parameters.ReplicationSystem->GetId());	
+	//UE_LOG_REPWRITER(VeryVerbose, TEXT("FReplicationWriter::UpdateDirtyChangeMasks() Updated %u Objects for ConnectionId:%u, ReplicationSystemId: %u."), CachedChangeMasks.Indices.Num(), Parameters.ConnectionId, Parameters.ReplicationSystem->GetId());	
 }
 
 void FReplicationWriter::NotifyDestroyedObjectPendingEndReplication(FInternalNetRefIndex ObjectInternalIndex)
@@ -1208,7 +1213,7 @@ void FReplicationWriter::HandleDeliveredRecord(const FReplicationRecord::FRecord
 	
 	if (CurrentState == EReplicatedObjectState::Invalid)
 	{
-		UE_LOG_REPLICATIONWRITER_WARNING(TEXT("FReplicationWriter::HandleDeliveredRecord - Warning Object ( InternalIndex: %u ) is invalid. DeliveredState %s WasDestroySubObject: %u"), InternalIndex, LexToString(DeliveredState), RecordInfo.WroteDestroySubObject)
+		UE_LOG_REPWRITER_WARNING(TEXT("FReplicationWriter::HandleDeliveredRecord - Warning Object ( InternalIndex: %u ) is invalid. DeliveredState %s WasDestroySubObject: %u"), InternalIndex, LexToString(DeliveredState), RecordInfo.WroteDestroySubObject)
 		ensure(false);
 		return;
 	}
@@ -1222,12 +1227,12 @@ void FReplicationWriter::HandleDeliveredRecord(const FReplicationRecord::FRecord
 		if (Info.LastAckedBaselineIndex != FDeltaCompressionBaselineManager::InvalidBaselineIndex)
 		{			
 			BaselineManager->DestroyBaseline(Parameters.ConnectionId, InternalIndex, Info.LastAckedBaselineIndex);
-			UE_LOG_REPLICATIONWRITER_CONN(TEXT("Destroyed old baseline %u for ( InternalIndex: %u )"), Info.LastAckedBaselineIndex, InternalIndex);			
+			UE_LOG_REPWRITER(VeryVerbose, TEXT("Destroyed old baseline %u for ( InternalIndex: %u )"), Info.LastAckedBaselineIndex, InternalIndex);			
 		}
 		Info.LastAckedBaselineIndex = RecordInfo.NewBaselineIndex;
 		Info.PendingBaselineIndex = FDeltaCompressionBaselineManager::InvalidBaselineIndex;
 
-		UE_LOG_REPLICATIONWRITER_CONN(TEXT("Acknowledged baseline %u for ( InternalIndex: %u )"), RecordInfo.NewBaselineIndex, InternalIndex);
+		UE_LOG_REPWRITER(VeryVerbose, TEXT("Acknowledged baseline %u for ( InternalIndex: %u )"), RecordInfo.NewBaselineIndex, InternalIndex);
 	}
 
 	// Update state
@@ -1272,13 +1277,13 @@ void FReplicationWriter::HandleDeliveredRecord(const FReplicationRecord::FRecord
 					// Tear-off is marked as a flush
 					if (Info.TearOff)
 					{
-						UE_LOG_REPLICATIONWRITER_CONN(TEXT("ReplicationWriter::HandleDeliveredRecord for ( InternalIndex: %u ) Waiting for flush before tearoff"), InternalIndex);
+						UE_LOG_REPWRITER(VeryVerbose, TEXT("ReplicationWriter::HandleDeliveredRecord for ( InternalIndex: %u ) Waiting for flush before tearoff"), InternalIndex);
 						SetState(InternalIndex, EReplicatedObjectState::WaitOnFlush);
 					}
 					// so are objects marked for destroy requiring flush
 					else if (ObjectsPendingDestroy.GetBit(InternalIndex))
 					{
-						UE_LOG_REPLICATIONWRITER_CONN(TEXT("ReplicationWriter::HandleDeliveredRecord for ( InternalIndex: %u ) Waiting for flush before destroy"), InternalIndex);
+						UE_LOG_REPWRITER(VeryVerbose, TEXT("ReplicationWriter::HandleDeliveredRecord for ( InternalIndex: %u ) Waiting for flush before destroy"), InternalIndex);
 						SetState(InternalIndex, EReplicatedObjectState::WaitOnFlush);
 					}
 				}
@@ -1755,13 +1760,13 @@ void FReplicationWriter::HandleDroppedRecord(const FReplicationRecord::FRecordIn
 
 	check(CurrentState != EReplicatedObjectState::Invalid);
 
-	UE_LOG_REPLICATIONWRITER_CONN(TEXT("Handle dropped data for ( InternalIndex: %u ) %s, LostState %s, CurrentState is %s"), InternalIndex, ToCStr(NetRefHandleManager->GetReplicatedObjectData(InternalIndex).RefHandle.ToString()), LexToString(LostObjectState), LexToString(CurrentState));
+	UE_LOG_REPWRITER(VeryVerbose, TEXT("Handle dropped data for ( InternalIndex: %u ) %s, LostState %s, CurrentState is %s"), InternalIndex, ToCStr(NetRefHandleManager->GetReplicatedObjectData(InternalIndex).RefHandle.ToString()), LexToString(LostObjectState), LexToString(CurrentState));
 
 	// If we loose a baseline we must notify the BaselineManager and invalidate our pendingbaselineindex
 	if (RecordInfo.NewBaselineIndex != FDeltaCompressionBaselineManager::InvalidBaselineIndex)
 	{
 		check(RecordInfo.NewBaselineIndex == Info.PendingBaselineIndex);
-		UE_LOG_REPLICATIONWRITER_CONN(TEXT("Lost baseline %u for ( InternalIndex: %u )"), RecordInfo.NewBaselineIndex, InternalIndex);
+		UE_LOG_REPWRITER(VeryVerbose, TEXT("Lost baseline %u for ( InternalIndex: %u )"), RecordInfo.NewBaselineIndex, InternalIndex);
 		
 		BaselineManager->LostBaseline(Parameters.ConnectionId, InternalIndex, RecordInfo.NewBaselineIndex);
 		Info.PendingBaselineIndex = FDeltaCompressionBaselineManager::InvalidBaselineIndex;
@@ -2237,6 +2242,26 @@ bool FReplicationWriter::CanSendObject(uint32 InternalIndex) const
 		}
 	}
 
+	// Disabled for now, this can be used if we want to enforce creation dependencies on the server side
+#if UE_NET_ENFORCE_CREATION_DEPENDENCIES_ON_SERVER
+	// Deal with initial dependencies
+	if ((State == EReplicatedObjectState::PendingCreate && NetRefHandleManager->GetObjectsWithCreationDependencies().GetBit(InternalIndex)))
+	{
+		for (const FInternalNetRefIndex& InitDependencyIndex : NetRefHandleManager->GetCreationDependencies(InternalIndex))
+		{
+			// Verify status of initial dependency to see if we can send this object
+			const FReplicationInfo& DependencyReplicationInfo = GetReplicationInfo(InitDependencyIndex);
+			if (!(DependencyReplicationInfo.IsCreationConfirmed || (WriteContext.ObjectsWrittenThisPacket.GetBit(InitDependencyIndex) && !IsActiveHugeObject(InitDependencyIndex))))
+			{
+				UE_LOG_REPWRITER(VeryVerbose, TEXT("ReplicationWriter: Cannot send internal index (%u) due to waiting on initial dependency internal index (%d)"), InternalIndex, InitDependencyIndex);
+				UE_NET_REPLICATIONWRITER_WARN_IF_CANNOT_SEND(TEXT("Waiting on init dependency %s "), *NetRefHandleManager->PrintObjectFromIndex(InitDependencyIndex));
+
+				return false;
+			}
+		}
+	}
+#endif
+
 	if (Info.HasDirtySubObjects)
 	{
 		for (FInternalNetRefIndex SubObjectInternalIndex : NetRefHandleManager->GetSubObjects(InternalIndex))
@@ -2439,7 +2464,7 @@ FReplicationWriter::EWriteObjectStatus FReplicationWriter::WriteObjectAndSubObje
 	{
 		if (CreatedBaselineIndex != FDeltaCompressionBaselineManager::InvalidBaselineIndex)
 		{
-			UE_LOG_REPLICATIONWRITER_CONN(TEXT("Destroy cancelled baseline %u for ( InternalIndex: %u )"), CreatedBaselineIndex, InternalIndex);
+			UE_LOG_REPWRITER(VeryVerbose, TEXT("Destroy cancelled baseline %u for ( InternalIndex: %u )"), CreatedBaselineIndex, InternalIndex);
 			BaselineManager->DestroyBaseline(Parameters.ConnectionId, InternalIndex, CreatedBaselineIndex);
 		}
 	};
@@ -2575,7 +2600,7 @@ FReplicationWriter::EWriteObjectStatus FReplicationWriter::WriteObjectAndSubObje
 							ApplyFilterToChangeMask(OutBatchInfo.ParentInternalIndex, InternalIndex, Info, ObjectData.Protocol, ReplicatedObjectStateBuffer, bIsInitialState);
 						}
 
-						UE_LOG_REPLICATIONWRITER_CONN(TEXT("Created new baseline %u for ( InternalIndex: %u )"), CreatedBaselineIndex, InternalIndex);
+						UE_LOG_REPWRITER(VeryVerbose, TEXT("Created new baseline %u for ( InternalIndex: %u )"), CreatedBaselineIndex, InternalIndex);
 					}
 				}
 			}
@@ -2603,7 +2628,7 @@ FReplicationWriter::EWriteObjectStatus FReplicationWriter::WriteObjectAndSubObje
 					// Warn if we cannot replicate this object
 					if (!ObjectData.Protocol || (!ObjectData.InstanceProtocol && !(ObjectData.bHasCachedCreationInfo == 1U)))
 					{
-						UE_LOG_REPLICATIONWRITER_WARNING(TEXT("Failed to replicate ( InternalIndex: %u ) %s, ProtocolName: %s, InstanceProtocol pointer: %p, HasCachedCreationInfo: %u"), InternalIndex, *NetRefHandle.ToString(), (ObjectData.Protocol ? ToCStr(ObjectData.Protocol->DebugName) : TEXT("nullptr")), ObjectData.InstanceProtocol, ObjectData.bHasCachedCreationInfo);
+						UE_LOG_REPWRITER_WARNING(TEXT("Failed to replicate ( InternalIndex: %u ) %s, ProtocolName: %s, InstanceProtocol pointer: %p, HasCachedCreationInfo: %u"), InternalIndex, *NetRefHandle.ToString(), (ObjectData.Protocol ? ToCStr(ObjectData.Protocol->DebugName) : TEXT("nullptr")), ObjectData.InstanceProtocol, ObjectData.bHasCachedCreationInfo);
 						ensureMsgf(ObjectData.Protocol, TEXT("Failed to replicate ( InternalIndex: %u ) %s, Protocol: nullptr, InstanceProtocol pointer: %p, HasCachedCreationInfo: %u"), InternalIndex, *NetRefHandle.ToString(), ObjectData.InstanceProtocol, ObjectData.bHasCachedCreationInfo);
 						return EWriteObjectStatus::NoInstanceProtocol;
 					}
@@ -2632,7 +2657,7 @@ FReplicationWriter::EWriteObjectStatus FReplicationWriter::WriteObjectAndSubObje
 					{
 						if (!Context.HasErrorOrOverflow())
 						{
-							UE_LOG_REPLICATIONWRITER_WARNING(TEXT("Failed to replicate ( InternalIndex: %u ) %s, ProtocolName: %s, InstanceProtocol pointer: %p, HasCachedCreationInfo: %u"), InternalIndex, *NetRefHandle.ToString(), (ObjectData.Protocol ? ToCStr(ObjectData.Protocol->DebugName) : TEXT("nullptr")), ObjectData.InstanceProtocol, ObjectData.bHasCachedCreationInfo);
+							UE_LOG_REPWRITER_WARNING(TEXT("Failed to replicate ( InternalIndex: %u ) %s, ProtocolName: %s, InstanceProtocol pointer: %p, HasCachedCreationInfo: %u"), InternalIndex, *NetRefHandle.ToString(), (ObjectData.Protocol ? ToCStr(ObjectData.Protocol->DebugName) : TEXT("nullptr")), ObjectData.InstanceProtocol, ObjectData.bHasCachedCreationInfo);
 							ensureMsgf(ObjectData.Protocol, TEXT("Failed to replicate ( InternalIndex: %u ) %s, Protocol: nullptr, InstanceProtocol pointer: %p, HasCachedCreationInfo: %u"), InternalIndex, *NetRefHandle.ToString(), ObjectData.InstanceProtocol, ObjectData.bHasCachedCreationInfo);
 
 							// Unforced error, treat it as we have no instance and cannot create this object but we can continue with other objects
@@ -3445,9 +3470,50 @@ uint32 FReplicationWriter::WriteObjects(FNetSerializationContext& Context)
 	// Write only if there are objects to send and we're not in replication record starvation mode.
 	bool bContinue = WriteContext.bHasUpdatedObjectsToSend && !WriteContext.bIsInReplicationRecordStarvation;
 
-	auto SendObjectFunction = [this, &Context, &WrittenObjectCount](FInternalNetRefIndex InternalIndex) -> EWriteStatus
+	auto SendObjectFunction = [this, &Context, &WrittenObjectCount](FInternalNetRefIndex InternalIndex, bool bAppendCreationDependencies) -> EWriteStatus
 	{
-		if (this->WriteContext.ObjectsWrittenThisPacket.GetBit(InternalIndex) || !this->CanSendObject(InternalIndex))
+		if (this->WriteContext.ObjectsWrittenThisPacket.GetBit(InternalIndex) || WriteContext.ObjectsSkippedInThisPacket.Contains(InternalIndex))
+		{
+			// From a logic point of view we consider this object written if it is either successfully written to packet or if it is skipped and we will retry if we have more packets
+			return EWriteStatus::Written;
+		}
+	
+		// Deal with initial dependencies
+		// If we have creationDependencies reschedule current object to after CreationDependencies to avoid unnecessary re-orderding waiting on the receiving end. 
+		if (bAppendCreationDependencies)
+		{
+			const FReplicationInfo& Info = GetReplicationInfo(InternalIndex);
+			const EReplicatedObjectState State = Info.GetState();
+
+			if (State == EReplicatedObjectState::PendingCreate && NetRefHandleManager->GetObjectsWithCreationDependencies().GetBit(InternalIndex))
+			{
+				TArray<FExtraScheduleInfo> ScheduledInitialDependencies;
+				for (FInternalNetRefIndex InitDependencyIndex : NetRefHandleManager->GetCreationDependencies(InternalIndex))
+				{
+					// Verify status of initial dependency to see if we can send this object
+					const FReplicationInfo& DependencyReplicationInfo = GetReplicationInfo(InitDependencyIndex);
+					if (!(DependencyReplicationInfo.IsCreationConfirmed || (WriteContext.ObjectsWrittenThisPacket.GetBit(InitDependencyIndex) && !IsActiveHugeObject(InitDependencyIndex))))
+					{
+						UE_LOG_REPWRITER(VeryVerbose, TEXT("ReplicationWriter:Scheduling initial dependency of internal index (%u) dependency internal index (%d)"), InternalIndex, InitDependencyIndex);
+						ScheduledInitialDependencies.Add(FExtraScheduleInfo(InitDependencyIndex));
+					}
+				}
+
+				if (!ScheduledInitialDependencies.IsEmpty())
+				{
+					// Schedule object after creation dependencies, this time without allowing it to schedule creation dependencies
+					const bool bAllowScheduleCreationDependencies = false;
+					ScheduledInitialDependencies.Add(FExtraScheduleInfo(InternalIndex, bAllowScheduleCreationDependencies));
+
+					// We add them as one operation as additionally scheduled objects is used as a stack
+					WriteContext.PushAdditionalScheduledObjects(ScheduledInitialDependencies);
+						
+					return EWriteStatus::RescheduledDueToCreationDependency;
+				}
+			}
+		}
+
+		if (!this->CanSendObject(InternalIndex))
 		{
 			// Even if the object cannot be sent, consider it written.
 			return EWriteStatus::Written;
@@ -3463,7 +3529,6 @@ uint32 FReplicationWriter::WriteObjects(FNetSerializationContext& Context)
 		return Result.Status;
 	};
 	
-
 	while (bContinue && ObjectListIt < ObjectCount)
 	{
 		// Partial sort next batch
@@ -3474,25 +3539,26 @@ uint32 FReplicationWriter::WriteObjects(FNetSerializationContext& Context)
 
 		for (;;)
 		{
-			// Try to send dependent objects
-			while (WriteContext.DependentObjectsPendingSend.Num() > 0)
+			// Process objects that has been explicitly scheduled due to dependencies etc.
+			while (bContinue && WriteContext.AdditionalObjectSchedulingStack.Num() > 0)
 			{
 				UE_NET_TRACE_SCOPE(DependentObjectData, *Context.GetBitStreamWriter(), Context.GetTraceCollector(), ENetTraceVerbosity::VeryVerbose);
 
-				const uint32 InternalIndex = WriteContext.DependentObjectsPendingSend.Pop();
-				checkSlow(InternalIndex != ObjectIndexForOOBAttachment);
-				ensureMsgf(GetReplicationInfo(InternalIndex).GetState() != EReplicatedObjectState::Invalid, TEXT("DependentObject with InternalIndex %u is not in scope"), InternalIndex);
+				const FExtraScheduleInfo ExtraScheduleInfo = WriteContext.AdditionalObjectSchedulingStack.Pop();
 
-				//UE_LOG(LogIris, Verbose, TEXT("Writing DependentObjectsPendingSend: %s"), *NetRefHandleManager->PrintObjectFromIndex(InternalIndex));
+				checkSlow(ExtraScheduleInfo.InternalIndex != ObjectIndexForOOBAttachment);
+				ensureMsgf(GetReplicationInfo(ExtraScheduleInfo.InternalIndex).GetState() != EReplicatedObjectState::Invalid, TEXT("DependentObject with InternalIndex %u is not in scope"), ExtraScheduleInfo.InternalIndex);
 
-				const EWriteStatus Status = SendObjectFunction(InternalIndex);
+				//UE_LOG_REPWRITER(VeryVerbose, TEXT("Writing AdditionalObjectSchedulingStack: %s"), *NetRefHandleManager->PrintObjectFromIndex(ExtraScheduleInfo.InternalIndex));
 
+				const EWriteStatus Status = SendObjectFunction(ExtraScheduleInfo.InternalIndex, ExtraScheduleInfo.bAppendCreationDependencies);
 				if (Status == EWriteStatus::Abort)
 				{
-					//UE_LOG(LogIris, Verbose, TEXT("Aborted writing of DependentObjectsPendingSend: %s. Added back to the list."), *NetRefHandleManager->PrintObjectFromIndex(InternalIndex));
+					//UE_LOG_REPWRITER(VeryVerbose, TEXT("Aborted writing of additionally scheduled object: %s. Added back to the stack."), *NetRefHandleManager->PrintObjectFromIndex(ExtraScheduleInfo.InternalIndex));
+
 					// If we fail, we put the object back on the pending send stack and try again in the next packet of the batch
 					// The reason for pop before using the index is that the SendObjectFunction will push new dependent objects on the stack
-					WriteContext.DependentObjectsPendingSend.Push(InternalIndex);
+					WriteContext.AdditionalObjectSchedulingStack.Push(ExtraScheduleInfo);
 					bContinue = false;
 					break;
 				}
@@ -3500,25 +3566,25 @@ uint32 FReplicationWriter::WriteObjects(FNetSerializationContext& Context)
 				// If we overflown and could not fit the object in the packet
 				if (Status == EWriteStatus::Skipped)
 				{
-					//UE_LOG(LogIris, Verbose, TEXT("Skipping over DependentObjectsPendingSend: %s"), *NetRefHandleManager->PrintObjectFromIndex(InternalIndex));
-					// TODO: We need to ensure it's ok to skip dependant objects and wait for the next frame to resend them if want want Dependent objects to be truly atomic
+					//UE_LOG_REPWRITER(VeryVerbose, TEXT("Skipping over additional scheduled object: %s"), *NetRefHandleManager->PrintObjectFromIndex(ExtraScheduleInfo.InternalIndex));
+					WriteContext.ObjectsSkippedInThisPacket.AddUnique(ExtraScheduleInfo.InternalIndex);
 				}
 			}
 
-			// Normal send
-			if (ObjectListIt < SortedCount)
+			// Pick next object from original scheduling list.
+			if (bContinue && ObjectListIt < SortedCount)
 			{
-				const uint32 InternalIndex = ObjectList[ObjectListIt].Index;
+				const FInternalNetRefIndex InternalIndex = ObjectList[ObjectListIt].Index;
 
 				checkSlow(InternalIndex != ObjectIndexForOOBAttachment);
 
-				//UE_LOG(LogIris, Verbose, TEXT("Writing regular object: %s"), *NetRefHandleManager->PrintObjectFromIndex(InternalIndex));
-
-				const EWriteStatus Status = SendObjectFunction(InternalIndex);
+				// If the object has creation dependencies we will schedule unsent creation dependencies before we try to send the object again.
+				const bool bAppendCreationDependencies = bScheduleCreationDependenciesFirst;
+				const EWriteStatus Status = SendObjectFunction(InternalIndex, bAppendCreationDependencies);
 				
 				if (Status == EWriteStatus::Abort)
 				{
-					//UE_LOG(LogIris, Verbose, TEXT("Aborted writing of regular object: %s"), *NetRefHandleManager->PrintObjectFromIndex(InternalIndex));
+					//UE_LOG_REPWRITER(VeryVerbose, TEXT("Aborted writing of regular object: %s"), *NetRefHandleManager->PrintObjectFromIndex(InternalIndex));
 					bContinue = false;
 					break;
 				}
@@ -3527,7 +3593,8 @@ uint32 FReplicationWriter::WriteObjects(FNetSerializationContext& Context)
 				{
 					// Continue trying to send objects as we might find an object with a smaller footprint that fits in the available buffer space.
 					// Note that skipped objects won't be reserialized again even if the more packets are available after this batch.
-					//UE_LOG(LogIris, Verbose, TEXT("Skipped writing of regular object: %s"), *NetRefHandleManager->PrintObjectFromIndex(InternalIndex));
+					//UE_LOG_REPWRITER(VeryVerbose, TEXT("Skipped writing of regular object: %s"), *NetRefHandleManager->PrintObjectFromIndex(InternalIndex));
+					WriteContext.ObjectsSkippedInThisPacket.AddUnique(InternalIndex);
 				}
 
 				++ObjectListIt;
@@ -3540,8 +3607,24 @@ uint32 FReplicationWriter::WriteObjects(FNetSerializationContext& Context)
 		}
 	}
 
+	// If we skipped objects, we move them to the extra scheduling stack so that they can be written in following packets if bandwidth allows.
+	if (WriteContext.ObjectsSkippedInThisPacket.Num())
+	{
+		// We add them as one operation as additionally scheduled objects is used as a stack
+		TArray<FExtraScheduleInfo, TInlineAllocator<32>> RescheduledSkippedObjects;
+		for (FInternalNetRefIndex InternalIndex : WriteContext.ObjectsSkippedInThisPacket)
+		{
+			RescheduledSkippedObjects.Add(FExtraScheduleInfo(InternalIndex));
+		}
+		
+		WriteContext.PushAdditionalScheduledObjects(RescheduledSkippedObjects);
+
+		// Reset so that we can try again in next packet sent this frame if allowed
+		WriteContext.ObjectsSkippedInThisPacket.Reset();
+	}
+
 	// if we have more data to write, store state so that we can continue if we are allowed to write more data
-	WriteContext.bHasUpdatedObjectsToSend = (ObjectListIt != ObjectCount) || (WriteContext.DependentObjectsPendingSend.Num() > 0U);
+	WriteContext.bHasUpdatedObjectsToSend = (ObjectListIt != ObjectCount) || (WriteContext.AdditionalObjectSchedulingStack.Num() > 0U);
 	WriteContext.CurrentIndex = ObjectListIt;
 	WriteContext.SortedObjectCount = SortedCount;
 
@@ -3564,7 +3647,7 @@ uint32 FReplicationWriter::HandleObjectBatchSuccess(const FBatchInfo& BatchInfo,
 	OutRecord.ObjectReplicationRecords.Reserve(BatchInfo.ObjectInfos.Num());
 	for (const FBatchObjectInfo& BatchObjectInfo : BatchInfo.ObjectInfos)
 	{
-		//UE_LOG_REPLICATIONWRITER(TEXT("FReplicationWriter::Wrote Object with %s InitialState: %u"), *NetRefHandleManager->NetHandle.ToString(), bIsInitialState ? 1u : 0u);
+		//UE_LOG_REPWRITER(VeryVerbose, TEXT("FReplicationWriter::Wrote Object with %s InitialState: %u"), *NetRefHandleManager->NetHandle.ToString(), bIsInitialState ? 1u : 0u);
 		FReplicationInfo& Info = GetReplicationInfo(BatchObjectInfo.InternalIndex);
 
 		// We did write the initial state, change the state to WaitOnCreateConfirmation
@@ -3581,7 +3664,7 @@ uint32 FReplicationWriter::HandleObjectBatchSuccess(const FBatchInfo& BatchInfo,
 			}
 			else
 			{
-				UE_LOG_REPLICATIONWRITER_CONN(TEXT("ReplicationWriter::HandleObjectBatchSuccess for ( InternalIndex: %u ) Waiting for flush before tearoff"), BatchObjectInfo.InternalIndex);
+				UE_LOG_REPWRITER(Verbose, TEXT("ReplicationWriter::HandleObjectBatchSuccess for ( InternalIndex: %u ) Waiting for flush before tearoff"), BatchObjectInfo.InternalIndex);
 				SetState(BatchObjectInfo.InternalIndex, EReplicatedObjectState::WaitOnFlush);
 			}
 		}
@@ -3621,16 +3704,21 @@ uint32 FReplicationWriter::HandleObjectBatchSuccess(const FBatchInfo& BatchInfo,
 		}
 
 		// Schedule rest of dependent objects for replication, note there is no guarantee that they will replicate in same packet
+		TArray<FExtraScheduleInfo> DependentObjects;
 		for (const FDependentObjectInfo DependentObjectInfo : NetRefHandleManager->GetDependentObjectInfos(BatchObjectInfo.InternalIndex))
 		{
 			const FInternalNetRefIndex DependentInternalIndex = DependentObjectInfo.NetRefIndex;
 			if (ObjectsWithDirtyChanges.GetBit(DependentInternalIndex) && !WriteContext.ObjectsWrittenThisPacket.GetBit(DependentInternalIndex))
 			{
-				WriteContext.DependentObjectsPendingSend.Push(DependentInternalIndex);
-				// Bumping the scheduling priority here will make sure that they will be scheduled the next update if we are not allowed to replicate this frame
+				DependentObjects.Add(FExtraScheduleInfo(DependentInternalIndex));
+				// Bumping the scheduling priority here will make sure that they will be picked up by scheduling the next tick/frame even if not all are allowed to replicate this frame
 				SchedulingPriorities[DependentInternalIndex] = FMath::Max(SchedulingPriorities[BatchObjectInfo.InternalIndex], SchedulingPriorities[DependentInternalIndex]);
 			}
-		}			
+		}
+		if (DependentObjects.Num())
+		{
+			WriteContext.PushAdditionalScheduledObjects(DependentObjects);
+		}
 
 		if (BatchObjectInfo.bSentState | BatchObjectInfo.bSentAttachments | BatchObjectInfo.bSentTearOff | BatchObjectInfo.bSentDestroySubObject)
 		{
@@ -3666,7 +3754,7 @@ uint32 FReplicationWriter::HandleObjectBatchSuccess(const FBatchInfo& BatchInfo,
 
 			if (UE_LOG_ACTIVE(LogIris, Verbose))
 			{
-				UE_LOG_REPLICATIONWRITER_CONN(TEXT("Blocked Object %s was sent after waiting for %f s"), *NetRefHandleManager->PrintObjectFromIndex(BatchObjectInfo.InternalIndex), FPlatformTime::ToSeconds64(FPlatformTime::Cycles64() - CannotSendInfo.StartCycles));
+				UE_LOG_REPWRITER(Verbose, TEXT("Blocked Object %s was sent after waiting for %f s"), *NetRefHandleManager->PrintObjectFromIndex(BatchObjectInfo.InternalIndex), FPlatformTime::ToSeconds64(FPlatformTime::Cycles64() - CannotSendInfo.StartCycles));
 			}
 		}
 #endif
@@ -3734,7 +3822,7 @@ FReplicationWriter::EWriteObjectRetryMode FReplicationWriter::HandleObjectBatchF
 		const uint32 SplitThreshold = PartialNetObjectAttachmentHandler->GetConfig()->GetBitCountSplitThreshold();
 		if (BitsLeft > SplitThreshold)
 		{
-			//UE_LOG_REPLICATIONWRITER(TEXT("FReplicationWriter::HandleObjectBatchFailure Failed to write object with ParentInternalIndex: %u EWriteObjectRetryMode::SplitHugeObject"), BatchInfo.ParentInternalIndex);
+			//UE_LOG_REPWRITER(VeryVerbose, TEXT("FReplicationWriter::HandleObjectBatchFailure Failed to write object with ParentInternalIndex: %u EWriteObjectRetryMode::SplitHugeObject"), BatchInfo.ParentInternalIndex);
 			return EWriteObjectRetryMode::SplitHugeObject;
 		}
 	}
@@ -3756,7 +3844,7 @@ FReplicationWriter::EWriteObjectRetryMode FReplicationWriter::HandleObjectBatchF
 	}
 
 	// Default- try some more, hopefully smaller state, objects.
-	//UE_LOG_REPLICATIONWRITER(TEXT("FReplicationWriter::HandleObjectBatchFailure Failed to write object with ParentInternalIndex: %u EWriteObjectRetryMode::TrySmallObject"), BatchInfo.ParentInternalIndex);
+	///UE_LOG_REPWRITER(VeryVerbose,TEXT("FReplicationWriter::HandleObjectBatchFailure Failed to write object with ParentInternalIndex: %u EWriteObjectRetryMode::TrySmallObject"), BatchInfo.ParentInternalIndex);
 	return EWriteObjectRetryMode::TrySmallObject;
 }
 
@@ -3825,8 +3913,8 @@ UDataStream::EWriteResult FReplicationWriter::BeginWrite(const UDataStream::FBeg
 		WriteContext.ScheduledObjectCount = ScheduleObjects(WriteContext.ScheduledObjectInfos);
 	}
 
-	// Reset dependent object array
-	WriteContext.DependentObjectsPendingSend.Reset();
+	// Reset AdditionalObjectSchedulingStack
+	WriteContext.AdditionalObjectSchedulingStack.Reset();
 
 	WriteContext.CurrentIndex = 0U;
 	WriteContext.FailedToWriteSmallObjectCount = 0U;
@@ -3877,28 +3965,28 @@ void FReplicationWriter::EndWrite()
 		}
 #endif
 
-		const uint32 NumPendingDependentObjects = WriteContext.DependentObjectsPendingSend.Num();
+		const uint32 NumPendingAdditionallySchedule1dObjects = WriteContext.AdditionalObjectSchedulingStack.Num();
 		const uint32 NumPendingObjectsToWrite = WriteContext.ScheduledObjectCount - WriteContext.CurrentIndex;
 
 		FReplicationStats& ReplicationStats = Parameters.ReplicationSystem->GetReplicationSystemInternal()->GetTickReplicationStats();
 		ReplicationStats.PendingObjectCount += NumPendingObjectsToWrite;
-		ReplicationStats.PendingDependentObjectCount += NumPendingDependentObjects;
+		ReplicationStats.PendingDependentObjectCount += NumPendingAdditionallySchedule1dObjects;
 		ReplicationStats.HugeObjectSendQueue += HugeObjectSendQueue.NumRootObjectsInTransit();
 		ReplicationStats.MaxPendingObjectCount = FMath::Max(ReplicationStats.MaxPendingObjectCount, NumPendingObjectsToWrite);
-		ReplicationStats.MaxPendingDependentObjectCount = FMath::Max(ReplicationStats.MaxPendingDependentObjectCount, NumPendingDependentObjects);
+		ReplicationStats.MaxPendingDependentObjectCount = FMath::Max(ReplicationStats.MaxPendingDependentObjectCount, NumPendingAdditionallySchedule1dObjects);
 		ReplicationStats.MaxHugeObjectSendQueue = FMath::Max(ReplicationStats.MaxHugeObjectSendQueue, HugeObjectSendQueue.NumRootObjectsInTransit());
 		++ReplicationStats.SampleCount;
 
  #if UE_NET_ENABLE_REPLICATIONWRITER_LOG
 		// See if we failed to write any objects
-		if (NumPendingDependentObjects)
+		if (NumPendingAdditionallySchedule1dObjects)
 		{
-			UE_LOG_REPLICATIONWRITER_WARNING(TEXT("FReplicationWriter::EndWrite() Has %u more dependent objects to write"), NumPendingDependentObjects);
+			UE_LOG_REPWRITER(VeryVerbose, TEXT("FReplicationWriter::EndWrite() Has %u more scheduled objects to write"), NumPendingAdditionallySchedule1dObjects);
 		}
 		
 		if (NumPendingObjectsToWrite)
 		{
-			UE_LOG_REPLICATIONWRITER_WARNING(TEXT("FReplicationWriter::EndWrite() Has %u more objects with dirty data"), NumPendingObjectsToWrite);
+			UE_LOG_REPWRITER(VeryVerbose, TEXT("FReplicationWriter::EndWrite() Has %u more objects with dirty data"), NumPendingObjectsToWrite);
 		}
 #endif
 
@@ -4026,7 +4114,7 @@ UDataStream::EWriteResult FReplicationWriter::Write(FNetSerializationContext& Co
 			Writer.WriteBits(TotalWrittenBatchCount, 16);
 		}
 
-		//UE_LOG_REPLICATIONWRITER(TEXT("FReplicationWriter::Write() Wrote %u Objects for ConnectionId:%u, ReplicationSystemId: %u."), WrittenObjectCount, Parameters.ConnectionId, Parameters.ReplicationSystem->GetId());	
+		//UE_LOG_REPWRITER(VeryVerbose,TEXT("FReplicationWriter::Write() Wrote %u Objects for ConnectionId:%u, ReplicationSystemId: %u."), WrittenObjectCount, Parameters.ConnectionId, Parameters.ReplicationSystem->GetId());	
 	
 		// Push record
 		const uint16 ReplicationInfoCount = static_cast<uint16>(ReplicationRecord.GetInfoCount() - OldReplicationInfoCount);
diff --git a/Engine/Source/Runtime/Net/Iris/Private/Iris/ReplicationSystem/ReplicationWriter.h b/Engine/Source/Runtime/Net/Iris/Private/Iris/ReplicationSystem/ReplicationWriter.h
--- a/Engine/Source/Runtime/Net/Iris/Private/Iris/ReplicationSystem/ReplicationWriter.h
+++ b/Engine/Source/Runtime/Net/Iris/Private/Iris/ReplicationSystem/ReplicationWriter.h
@@ -224,15 +224,41 @@ private:
 		float SortKey;
 	};
 
+	struct FExtraScheduleInfo
+	{
+		FExtraScheduleInfo(FInternalNetRefIndex InInternalIndex, bool bInAllowAppendCreationDependencies = true)
+		: InternalIndex(InInternalIndex) 
+		, bAppendCreationDependencies(bInAllowAppendCreationDependencies)
+		{			
+		}
+
+		FInternalNetRefIndex InternalIndex = 0U;
+		bool bAppendCreationDependencies = false;
+	};
+
 	// We persist some information during write so that we can support writing multiple packets with data without re-doing scheduling work.
 	struct FWriteContext
 	{
 		FWriteContext() : bIsValid(0) {}
 
-		// Objects we have written in this packet batch to avoid sending same object multiple times
+		// Objects we have written in this packet to avoid sending same object multiple times
 		FNetBitArray ObjectsWrittenThisPacket;
-		// DependentObjets that we should try to write this packet batch, aka. allow overcommit if we have pending DependentObjects when the packet is full
-		TArray<uint32, TInlineAllocator<32>> DependentObjectsPendingSend;
+
+		// Push a list of objects to the AdditionalObjectSchedulingStack. It will be added in reverse to maintain order.
+		void PushAdditionalScheduledObjects(TArrayView<FExtraScheduleInfo> ObjectsToSchedule)
+		{ 
+			for (int32 It = ObjectsToSchedule.Num() - 1; It >= 0; It--)
+			{
+				AdditionalObjectSchedulingStack.Push(ObjectsToSchedule[It]);
+			}
+		}
+
+		// Additional objects that we have scheduled to be written, it could be dependencies, dependent objects.
+		TArray<FExtraScheduleInfo, TInlineAllocator<32>> AdditionalObjectSchedulingStack;
+
+		// Objects skipped this packet since they did not fit, that we want to retry next packet
+		TArray<FInternalNetRefIndex, TInlineAllocator<32>> ObjectsSkippedInThisPacket;
+
 		// Scheduled objects
 		FScheduleObjectInfo* ScheduledObjectInfos;
 		uint32 ScheduledObjectCount;
@@ -507,6 +533,8 @@ private:
 		Skipped = 0,
 		// Object was successfully written
 		Written = 1,
+		// Object is rescheduled due to creation dependency
+		RescheduledDueToCreationDependency = 2,
 	};
 
 	// Write destruction info for an object that should be destroyed
