Support additional work to be done when marking the property dirty conditional on the state changing
Fixes PVS warning V571: Recurring check. The 'if (...)' condition was already verified in line N.

diff --git a/Engine/Source/Runtime/Net/Core/Public/Net/Core/PushModel/PushModel.h b/Engine/Source/Runtime/Net/Core/Public/Net/Core/PushModel/PushModel.h
--- a/Engine/Source/Runtime/Net/Core/Public/Net/Core/PushModel/PushModel.h
+++ b/Engine/Source/Runtime/Net/Core/Public/Net/Core/PushModel/PushModel.h
@@ -463,8 +463,16 @@ namespace UEPushModelPrivate
  * This is used to reduce lines of code (mostly in setters) by doing the comparison check before changing the value and marking dirty only if we have changed the value. 
  * This is best if used for things like ints, floats, pointers, but not something like a struct because of the memcmp cost. 
  */
-#define COMPARE_ASSIGN_AND_MARK_PROPERTY_DIRTY(ClassName, PropertyName, NewValue, Object) \
-	if (NewValue != PropertyName) { PropertyName = NewValue; MARK_PROPERTY_DIRTY_FROM_NAME(ClassName, PropertyName, Object); }
+#define COMPARE_ASSIGN_AND_MARK_PROPERTY_DIRTY(ClassName, PropertyName, NewValue, Object, ...) \
+	do { \
+		if (NewValue != PropertyName) \
+		{ \
+			PropertyName = NewValue; \
+			MARK_PROPERTY_DIRTY_FROM_NAME(ClassName, PropertyName, Object); \
+			__VA_ARGS__ \
+		} \
+	} while (0)
+
 
 #define UE_NET_SET_IRIS_MARK_PROPERTY_DIRTY_DELEGATE(Delegate) UEPushModelPrivate::SetIrisMarkPropertyDirtyDelegate(Delegate)
 #define UE_NET_SET_IRIS_MARK_PROPERTIES_DIRTY_DELEGATE(Delegate) UEPushModelPrivate::SetIrisMarkPropertiesDirtyDelegate(Delegate)
@@ -487,8 +495,14 @@ namespace UEPushModelPrivate
 #define IS_PUSH_MODEL_ENABLED() false
 #define PUSH_MAKE_BP_PROPERTIES_PUSH_MODEL() false
 
-#define COMPARE_ASSIGN_AND_MARK_PROPERTY_DIRTY(ClassName, PropertyName, NewValue, Object) \
-	if (NewValue != PropertyName) { PropertyName = NewValue; }
+#define COMPARE_ASSIGN_AND_MARK_PROPERTY_DIRTY(ClassName, PropertyName, NewValue, Object, ...) \
+	do { \
+		if (NewValue != PropertyName) \
+		{ \
+			PropertyName = NewValue; \
+			__VA_ARGS__ \
+		} \
+	} while(0)
 
 #define UE_NET_SET_IRIS_MARK_PROPERTY_DIRTY_DELEGATE(...)
 #define UE_NET_SET_IRIS_MARK_PROPERTIES_DIRTY_DELEGATE(...)
