Added an option to use a bespoke identifier in FNetRefHandle and adjust it's size.

diff --git a/Engine/Plugins/Runtime/GameplayAbilities/Source/GameplayAbilities/Private/Serialization/GameplayAbilityTargetingLocationInfoNetSerializer.cpp b/Engine/Plugins/Runtime/GameplayAbilities/Source/GameplayAbilities/Private/Serialization/GameplayAbilityTargetingLocationInfoNetSerializer.cpp
--- a/Engine/Plugins/Runtime/GameplayAbilities/Source/GameplayAbilities/Private/Serialization/GameplayAbilityTargetingLocationInfoNetSerializer.cpp
+++ b/Engine/Plugins/Runtime/GameplayAbilities/Source/GameplayAbilities/Private/Serialization/GameplayAbilityTargetingLocationInfoNetSerializer.cpp
@@ -33,7 +33,7 @@ struct FGameplayAbilityTargetingLocationInfoNetSerializer
 	struct FQuantizedData
 	{
 		// Keep TargetingLocationInfo first. It simplifies things.
-		alignas(16) uint8 TargetingLocationInfo[192 - 1];
+		alignas(16) uint8 TargetingLocationInfo[208];
 
 		uint8 LocationType;
 	};
diff --git a/Engine/Plugins/Runtime/GameplayAbilities/Source/GameplayAbilities/Private/Serialization/GameplayEffectContextNetSerializer.cpp b/Engine/Plugins/Runtime/GameplayAbilities/Source/GameplayAbilities/Private/Serialization/GameplayEffectContextNetSerializer.cpp
--- a/Engine/Plugins/Runtime/GameplayAbilities/Source/GameplayAbilities/Private/Serialization/GameplayEffectContextNetSerializer.cpp
+++ b/Engine/Plugins/Runtime/GameplayAbilities/Source/GameplayAbilities/Private/Serialization/GameplayEffectContextNetSerializer.cpp
@@ -47,9 +47,9 @@ struct FGameplayEffectContextNetSerializer
 	struct FQuantizedType
 	{
 #if UE_WITH_REMOTE_OBJECT_HANDLE
-		alignas(16) uint8 EffectContext[176];
+		alignas(16) uint8 EffectContext[192];
 #else
-		alignas(16) uint8 EffectContext[152];
+		alignas(16) uint8 EffectContext[160];
 #endif
 		// 320 should be large enough and makes struct a multiple of 16 bytes
 		alignas(16) uint8 HitResult[320];
diff --git a/Engine/Plugins/Runtime/GameplayAbilities/Source/GameplayAbilities/Public/Serialization/GameplayEffectContextNetSerializer.h b/Engine/Plugins/Runtime/GameplayAbilities/Source/GameplayAbilities/Public/Serialization/GameplayEffectContextNetSerializer.h
--- a/Engine/Plugins/Runtime/GameplayAbilities/Source/GameplayAbilities/Public/Serialization/GameplayEffectContextNetSerializer.h
+++ b/Engine/Plugins/Runtime/GameplayAbilities/Source/GameplayAbilities/Public/Serialization/GameplayEffectContextNetSerializer.h
@@ -15,7 +15,7 @@ namespace UE::Net
 {
 
 #if UE_WITH_REMOTE_OBJECT_HANDLE
-constexpr SIZE_T GetGameplayEffectContextNetSerializerSafeQuantizedSize() { return 512; }
+constexpr SIZE_T GetGameplayEffectContextNetSerializerSafeQuantizedSize() { return 528; }
 #else
 constexpr SIZE_T GetGameplayEffectContextNetSerializerSafeQuantizedSize() { return 496; }
 #endif
diff --git a/Engine/Source/Runtime/CoreUObject/Public/UObject/RemoteObjectTypes.h b/Engine/Source/Runtime/CoreUObject/Public/UObject/RemoteObjectTypes.h
--- a/Engine/Source/Runtime/CoreUObject/Public/UObject/RemoteObjectTypes.h
+++ b/Engine/Source/Runtime/CoreUObject/Public/UObject/RemoteObjectTypes.h
@@ -365,7 +365,8 @@ public:
 	friend struct UE::Net::FRemoteObjectIdNetSerializer;
 	friend UE::RemoteObject::Private::FRemoteIdLocalizationHelper;
 	friend struct FUObjectItem;
-	friend class FNetGUIDCache;	
+	friend class FNetGUIDCache;
+	friend class UNetDriver;
 private:
 	
 	static bool IsAssetInternal(UObjectBase* InObject, EInternalObjectFlags InInitialFlags);
diff --git a/Engine/Source/Runtime/Engine/Private/NetDriver.cpp b/Engine/Source/Runtime/Engine/Private/NetDriver.cpp
--- a/Engine/Source/Runtime/Engine/Private/NetDriver.cpp
+++ b/Engine/Source/Runtime/Engine/Private/NetDriver.cpp
@@ -7503,6 +7503,28 @@ void UNetDriver::CreateReplicationSystem(bool bInitAsClient)
 			Params.InitialNetRefHandleIndex = GInitialNetworkId;
 		}
 
+#if UE_WITH_REMOTE_OBJECT_HANDLE
+		Params.ObjectToNetRefHandleSerial.BindLambda([](const TObjectPtr<const UObject>& ObjectPtr)
+		{
+			const FRemoteObjectId RemoteObjectId = ObjectPtr.GetRemoteId();
+			const uint64 RemoteObjectIdValue = RemoteObjectId.GetIdNumber();
+
+			// Shift away the reserved bit in FRemoteObjectId so we're just left with the server id and serial number
+			// which will be embedded in the NetRefHandle serial.
+			const uint64 Serial = (RemoteObjectIdValue >> 1);
+
+			return Serial;
+		});
+
+		// Disabled until the calling code in Iris has been transactionalized.
+		/*Params.NetRefHandleSerialToObject.BindLambda([](const uint64 Serial)
+		{
+			const FRemoteObjectId RemoteObjectId = FRemoteObjectId::CreateFromInt(Serial).GetLocalized();
+			TObjectPtr<UObject> RemoteObjectPtr = TObjectPtr<UObject>(FObjectPtr(RemoteObjectId));
+			return RemoteObjectPtr;
+		});*/
+#endif
+
 		SetReplicationSystem(UE::Net::FReplicationSystemFactory::CreateReplicationSystem(Params));
 	}
 	else
diff --git a/Engine/Source/Runtime/Net/Iris/Private/Iris/ReplicationSystem/NetRefHandleManager.cpp b/Engine/Source/Runtime/Net/Iris/Private/Iris/ReplicationSystem/NetRefHandleManager.cpp
--- a/Engine/Source/Runtime/Net/Iris/Private/Iris/ReplicationSystem/NetRefHandleManager.cpp
+++ b/Engine/Source/Runtime/Net/Iris/Private/Iris/ReplicationSystem/NetRefHandleManager.cpp
@@ -192,12 +192,9 @@ void FNetRefHandleManager::GrowNetChunkedArrayBuffers(FInternalNetRefIndex Inter
 	OnNetChunkedArrayIncrease.Broadcast(HighestNetChunkedArrayInternalIndex);
 }
 
-uint64 FNetRefHandleManager::GetNextNetRefHandleId(uint64 HandleId) const
+uint64 FNetRefHandleManager::GetNextNetRefHandleSerial(uint64 HandleId) const
 {
-	// Since we use the lowest bit in the index to indicate if the handle is static or dynamic we cannot use all bits as the index
-	constexpr uint64 NetHandleIdIndexBitMask = (1ULL << (FNetRefHandle::IdBits - 1)) - 1;
-
-	uint64 NextHandleId = (HandleId + 1) & NetHandleIdIndexBitMask;
+	uint64 NextHandleId = (HandleId + 1) % FNetRefHandle::MaxSerial;
 	if (NextHandleId == 0)
 	{
 		++NextHandleId;
@@ -338,10 +335,19 @@ bool FNetRefHandleManager::HasInstanceProtocol(FInternalNetRefIndex InternalInde
 
 FNetRefHandle FNetRefHandleManager::AllocateNetRefHandle(bool bIsStatic)
 {
-	uint64& NextHandleId = bIsStatic ? NextStaticHandleIndex : NextDynamicHandleIndex;
+	uint64& NextHandleSerial = bIsStatic ? NextStaticHandleIndex : NextDynamicHandleIndex;
+
+	FNetRefHandle NewHandle = AllocateNetRefHandleFromSerial(NextHandleSerial, bIsStatic);
+
+	// Bump NextHandleSerial
+	NextHandleSerial = GetNextNetRefHandleSerial(NextHandleSerial);
+
+	return NewHandle;
+}
 
-	const uint64 NewHandleId = MakeNetRefHandleId(NextHandleId, bIsStatic);
-	FNetRefHandle NewHandle = MakeNetRefHandle(NewHandleId, ReplicationSystemId);
+FNetRefHandle FNetRefHandleManager::AllocateNetRefHandleFromSerial(uint64 Serial, bool bIsStatic)
+{
+	FNetRefHandle NewHandle = MakeNetRefHandleFromSerial(Serial, bIsStatic, ReplicationSystemId);
 
 	// Verify that the handle is free
 	if (RefHandleToInternalIndex.Contains(NewHandle))
@@ -351,9 +357,6 @@ FNetRefHandle FNetRefHandleManager::AllocateNetRefHandle(bool bIsStatic)
 		return FNetRefHandle();
 	}
 
-	// Bump NextHandleId
-	NextHandleId = GetNextNetRefHandleId(NextHandleId);
-
 	return NewHandle;
 }
 
@@ -1165,20 +1168,16 @@ void FNetRefHandleManager::AddReferencedObjects(FReferenceCollector& Collector)
 	}
 }
 
-uint64 FNetRefHandleManager::MakeNetRefHandleId(uint64 Id, bool bIsStatic)
-{
-	return (Id << 1U) | (bIsStatic ? 1U : 0U);
-}
-
 FNetRefHandle FNetRefHandleManager::MakeNetRefHandle(uint64 Id, uint32 ReplicationSystemId)
 {
-	check((Id & FNetRefHandle::IdMask) == Id);
+	check(Id <= FNetRefHandle::MaxId);
 	check(ReplicationSystemId < FNetRefHandle::MaxReplicationSystemId);
 
 	FNetRefHandle Handle;
 
-	Handle.Id = Id;
-	Handle.ReplicationSystemId = ReplicationSystemId + 1U;
+	Handle.Fields.Static = Id & 1;
+	Handle.Fields.Serial = Id >> 1;
+	Handle.Fields.ReplicationSystemId = ReplicationSystemId + 1U;
 
 	return Handle;
 }
@@ -1186,12 +1185,26 @@ FNetRefHandle FNetRefHandleManager::MakeNetRefHandle(uint64 Id, uint32 Replicati
 FNetRefHandle FNetRefHandleManager::MakeNetRefHandleFromId(uint64 Id)
 {
 	// This is called on the receiving end when deserializing replicated objects. We don't want to crash on bit stream errors leading to invalid handle IDs being read.
-	ensure((Id & FNetRefHandle::IdMask) == Id);
+	ensure(Id <= FNetRefHandle::MaxId);
+
+	FNetRefHandle Handle;
+
+	Handle.Fields.Static = Id & 1;
+	Handle.Fields.Serial = Id >> 1;
+	Handle.Fields.ReplicationSystemId = 0U;
+
+	return Handle;
+}
+
+FNetRefHandle FNetRefHandleManager::MakeNetRefHandleFromSerial(uint64 Serial, bool bIsStatic, uint32 ReplicationSystemId)
+{
+	check(Serial <= FNetRefHandle::MaxSerial);
 
 	FNetRefHandle Handle;
 
-	Handle.Id = Id;
-	Handle.ReplicationSystemId = 0U;
+	Handle.Fields.Static = bIsStatic;
+	Handle.Fields.Serial = Serial;
+	Handle.Fields.ReplicationSystemId = ReplicationSystemId + 1U;
 
 	return Handle;
 }
diff --git a/Engine/Source/Runtime/Net/Iris/Private/Iris/ReplicationSystem/NetRefHandleManager.h b/Engine/Source/Runtime/Net/Iris/Private/Iris/ReplicationSystem/NetRefHandleManager.h
--- a/Engine/Source/Runtime/Net/Iris/Private/Iris/ReplicationSystem/NetRefHandleManager.h
+++ b/Engine/Source/Runtime/Net/Iris/Private/Iris/ReplicationSystem/NetRefHandleManager.h
@@ -177,8 +177,12 @@ public:
 
 	static FNetRefHandle MakeNetRefHandle(uint64 Id, uint32 ReplicationSystemId);
 	static FNetRefHandle MakeNetRefHandleFromId(uint64 Id);
+	static FNetRefHandle MakeNetRefHandleFromSerial(uint64 Serial, bool bIsStatic, uint32 ReplicationSystemId);
 
-	// Returns a valid handle if the wanted handle can be allocated
+	// Returns a valid handle using a bespoke identifier.
+	FNetRefHandle AllocateNetRefHandleFromSerial(uint64 Serial, bool bIsStatic);
+
+	// Returns a valid handle using an internally generated identifier.
 	FNetRefHandle AllocateNetRefHandle(bool bIsStatic);
 
 	// Create local Net Object
@@ -420,8 +424,7 @@ private:
 	/** Clear the status of all internal bit arrays when an internal index is put back in the free list */
 	void ClearStateForFreedInternalIndex(FInternalNetRefIndex FreedInternalIndex);
 
-	static uint64 MakeNetRefHandleId(uint64 Seed, bool bIsStatic);
-	uint64 GetNextNetRefHandleId(uint64 HandleIndex) const;
+	uint64 GetNextNetRefHandleSerial(uint64 HandleIndex) const;
 
 	// Get the next free internal index, returns InvalidInternalIndex if a free one cannot be found
 	FInternalNetRefIndex GetNextFreeInternalIndex() const;
diff --git a/Engine/Source/Runtime/Net/Iris/Private/Iris/ReplicationSystem/ObjectReferenceCache.cpp b/Engine/Source/Runtime/Net/Iris/Private/Iris/ReplicationSystem/ObjectReferenceCache.cpp
--- a/Engine/Source/Runtime/Net/Iris/Private/Iris/ReplicationSystem/ObjectReferenceCache.cpp
+++ b/Engine/Source/Runtime/Net/Iris/Private/Iris/ReplicationSystem/ObjectReferenceCache.cpp
@@ -120,9 +120,11 @@ FObjectReferenceCache::FObjectReferenceCache()
 	SetAsyncLoadMode(EAsyncLoadMode::UseCVar);
 }
 
-void FObjectReferenceCache::Init(bool bAuthority)
+void FObjectReferenceCache::Init(const FObjectReferenceCacheParams& Params)
 {
-	bIsAuthority = bAuthority;
+	bIsAuthority = Params.bAuthority;
+	ObjectToNetRefHandleSerial = Params.ObjectToNetRefHandleSerial;
+	NetRefHandleSerialToObject = Params.NetRefHandleSerialToObject;
 }
 
 bool FObjectReferenceCache::SupportsObjectInternal(const UObject* Object) const
@@ -425,7 +427,16 @@ bool FObjectReferenceCache::CreateObjectReferenceInternal(const UObject* Object,
 	}
 	else
 	{
-		NetRefHandle = NetRefHandleManager.AllocateNetRefHandle(bIsStatic);
+		if (ObjectToNetRefHandleSerial.IsBound())
+		{
+			TObjectPtr<const UObject> ObjectPtr(Object);
+			const uint64 Serial = ObjectToNetRefHandleSerial.Execute(ObjectPtr);
+			NetRefHandle = NetRefHandleManager.AllocateNetRefHandleFromSerial(Serial, bIsStatic);
+		}
+		else
+		{
+			NetRefHandle = NetRefHandleManager.AllocateNetRefHandle(bIsStatic);
+		}
 	}
 
 	check(NetRefHandle.IsValid());
@@ -1173,8 +1184,18 @@ ENetObjectReferenceResolveResult FObjectReferenceCache::ResolveObjectReference(c
 			// If both path and handle is valid this is a reference to an object with a relative path.
 			if (Reference.GetRefHandle().IsValid())
 			{
-				bool bLocalMustBeMapped = false;
-				UObject* ReplicatedOuter = ResolveObjectReferenceHandleInternal(Reference.GetRefHandle(), ResolveContext, bLocalMustBeMapped);
+				UObject* ReplicatedOuter = nullptr;
+
+				if (NetRefHandleSerialToObject.IsBound())
+				{
+					ReplicatedOuter = NetRefHandleSerialToObject.Execute(Reference.GetRefHandle().GetSerial());
+				}
+
+				if (ReplicatedOuter == nullptr)
+				{
+					bool bLocalMustBeMapped = false;
+					ReplicatedOuter = ResolveObjectReferenceHandleInternal(Reference.GetRefHandle(), ResolveContext, bLocalMustBeMapped);
+				}
 				ResolvedObject = ReplicatedOuter ? StaticFindObject(UObject::StaticClass(), ReplicatedOuter, ToCStr(ObjectPath), EFindObjectFlags::None) : nullptr;
 
 				if (ResolvedObject == nullptr)
diff --git a/Engine/Source/Runtime/Net/Iris/Private/Iris/ReplicationSystem/ObjectReferenceCache.h b/Engine/Source/Runtime/Net/Iris/Private/Iris/ReplicationSystem/ObjectReferenceCache.h
--- a/Engine/Source/Runtime/Net/Iris/Private/Iris/ReplicationSystem/ObjectReferenceCache.h
+++ b/Engine/Source/Runtime/Net/Iris/Private/Iris/ReplicationSystem/ObjectReferenceCache.h
@@ -31,6 +31,13 @@ namespace UE::Net
 namespace UE::Net::Private
 {
 
+struct FObjectReferenceCacheParams
+{
+	bool bAuthority = false;
+	FObjectToNetRefHandleSerial ObjectToNetRefHandleSerial;
+	FNetRefHandleSerialToObject NetRefHandleSerialToObject;
+};
+
 // A lot of code in this class is extracted from the re-factored GUIDCache/ClientPackageMap and must be kept up to sync once they are submitted
 // Hopefully we can merge parts of this back together later on
 
@@ -40,7 +47,7 @@ public:
 
 	FObjectReferenceCache();
 
-	void Init(bool bAuthority);
+	void Init(const FObjectReferenceCacheParams& Params);
 
 	// Determine if the object is dynamic
 	bool IsDynamicObject(const UObject* Object) const;
@@ -264,6 +271,9 @@ private:
 	/** Set of packages that are currently pending async loads, referenced by package name. */
 	TMap<FName, FPendingAsyncLoadRequest> PendingAsyncLoadRequests;
 
+	FObjectToNetRefHandleSerial ObjectToNetRefHandleSerial;
+	FNetRefHandleSerialToObject NetRefHandleSerialToObject;
+
 	// $TODO: $IRIS: Stats support
 #if 0
 	/** Store all GUIDs that caused the sync loading of a package, for debugging & logging with LogNetSyncLoads */
diff --git a/Engine/Source/Runtime/Net/Iris/Private/Iris/ReplicationSystem/ReplicationSystem.cpp b/Engine/Source/Runtime/Net/Iris/Private/Iris/ReplicationSystem/ReplicationSystem.cpp
--- a/Engine/Source/Runtime/Net/Iris/Private/Iris/ReplicationSystem/ReplicationSystem.cpp
+++ b/Engine/Source/Runtime/Net/Iris/Private/Iris/ReplicationSystem/ReplicationSystem.cpp
@@ -292,7 +292,13 @@ public:
 		// Init replication bridge
 		ReplicationSystemInternal.GetReplicationBridge()->Initialize(ReplicationSystem);
 
-		ReplicationSystemInternal.GetObjectReferenceCache().Init(ReplicationSystem->IsServer());
+		{
+			FObjectReferenceCacheParams InitParams;
+			InitParams.bAuthority = ReplicationSystem->IsServer();
+			InitParams.ObjectToNetRefHandleSerial = Params.ObjectToNetRefHandleSerial;
+			InitParams.NetRefHandleSerialToObject = Params.NetRefHandleSerialToObject;
+			ReplicationSystemInternal.GetObjectReferenceCache().Init(InitParams);
+		}
 
 		// Init custom packagemap we use for capturing references for backwards compatible NetSerializers
 		{
@@ -2172,7 +2178,8 @@ TSharedPtr<UE::Net::Private::FObjectReferenceCache> FReplicationSystemFactory::G
 		SharedObjectReferenceCache = MakeShared<UE::Net::Private::FObjectReferenceCache>();
 
 		// It's assumed that the shared object reference cache will be used by client replication systems.
-		SharedObjectReferenceCache->Init(false);
+		UE::Net::Private::FObjectReferenceCacheParams InitParams;
+		SharedObjectReferenceCache->Init(InitParams);
 	}
 
 	return SharedObjectReferenceCache;
diff --git a/Engine/Source/Runtime/Net/Iris/Private/Iris/Serialization/FieldPathNetSerializer.cpp b/Engine/Source/Runtime/Net/Iris/Private/Iris/Serialization/FieldPathNetSerializer.cpp
--- a/Engine/Source/Runtime/Net/Iris/Private/Iris/Serialization/FieldPathNetSerializer.cpp
+++ b/Engine/Source/Runtime/Net/Iris/Private/Iris/Serialization/FieldPathNetSerializer.cpp
@@ -39,7 +39,7 @@ struct FFieldPathNetSerializer
 	// Types
 	struct FQuantizedType
 	{
-		alignas(8) uint8 QuantizedStruct[40];
+		alignas(8) uint8 QuantizedStruct[48];
 	};
 
 	typedef FFieldPathForNetSerializer SourceType;
diff --git a/Engine/Source/Runtime/Net/Iris/Public/Iris/Core/BitTwiddling.h b/Engine/Source/Runtime/Net/Iris/Public/Iris/Core/BitTwiddling.h
--- a/Engine/Source/Runtime/Net/Iris/Public/Iris/Core/BitTwiddling.h
+++ b/Engine/Source/Runtime/Net/Iris/Public/Iris/Core/BitTwiddling.h
@@ -105,4 +105,20 @@ T GetLeastSignificantBit(const T Value)
 	return LeastSignificantBit;
 }
 
+/** Return the largest unsigned value that can be stored in BitCount numbers of bits. */
+constexpr inline uint64 LargestBitValueForBits(uint64 BitCount)
+{
+	if (BitCount == 0)
+	{
+		return 0ULL;
+	}
+
+	if (BitCount > 64)
+	{
+		return ~0ULL;
+	}
+
+	return ~0ULL >> (64 - BitCount);
+}
+
 }
diff --git a/Engine/Source/Runtime/Net/Iris/Public/Iris/ReplicationSystem/NetRefHandle.h b/Engine/Source/Runtime/Net/Iris/Public/Iris/ReplicationSystem/NetRefHandle.h
--- a/Engine/Source/Runtime/Net/Iris/Public/Iris/ReplicationSystem/NetRefHandle.h
+++ b/Engine/Source/Runtime/Net/Iris/Public/Iris/ReplicationSystem/NetRefHandle.h
@@ -4,9 +4,14 @@
 
 #include "CoreTypes.h"
 #include "Iris/IrisConfig.h"
+#include "Iris/Core/BitTwiddling.h"
 #include "Containers/StringFwd.h"
 #include "Templates/TypeHash.h"
 
+#ifndef UE_NET_IRIS_NETREFHANDLE_SERIAL_SIZE
+	#define UE_NET_IRIS_NETREFHANDLE_SERIAL_SIZE 53
+#endif
+
 // Forward declarations
 class FString;
 
@@ -28,55 +33,59 @@ public:
 	inline static FNetRefHandle GetInvalid() { return FNetRefHandle(); }
 
 private:
-	enum { InvalidValue = 0 };
-	enum { IdBits = 54 };
-	enum { ReplicationSystemIdBits = 10 };
+	static constexpr uint64 InvalidValue = 0;
+	static constexpr uint64 StaticBits = 1;
+	static constexpr uint64 SerialBits = UE_NET_IRIS_NETREFHANDLE_SERIAL_SIZE;
+	static constexpr uint64 IdBits = (StaticBits + SerialBits);
+	static constexpr uint64 ReplicationSystemIdBits = 10;
+
+	static_assert(IdBits <= 64, "The number of bits representing the FNetRefHandle id cannot be greater than 64-bits");
 
 public:
 
-	static constexpr uint64 MaxReplicationSystemId = (1ULL << ReplicationSystemIdBits) - 1;
+	static constexpr uint64 MaxReplicationSystemId = UE::Net::LargestBitValueForBits(ReplicationSystemIdBits);
 	static constexpr uint64 MaxReplicationSystemCount = MaxReplicationSystemId + 1;
+	static constexpr uint64 MaxId = UE::Net::LargestBitValueForBits(IdBits);
+	static constexpr uint64 MaxSerial = UE::Net::LargestBitValueForBits(SerialBits);
 
-	FNetRefHandle() : Value(InvalidValue) {}
-
-	uint64 GetId() const { return Id; }
-	uint32 GetReplicationSystemId() const { check(ReplicationSystemId != 0); return (uint32)(ReplicationSystemId - 1); }
-	bool IsValid() const { return Value != InvalidValue; }
+	uint64 GetId() const { return (Fields.Serial << 1) | (Fields.Static & 1); }
+	uint64 GetSerial() const { return Fields.Serial; }
+	uint32 GetReplicationSystemId() const { check(Fields.ReplicationSystemId != 0); return (uint32)(Fields.ReplicationSystemId - 1); }
+	bool IsValid() const { return Fields.Serial != InvalidValue; }
 
 	/** Does the handle know which ReplicationSystem it is related to. */
-	bool IsCompleteHandle() const { return Value != InvalidValue && ReplicationSystemId != 0U; }
+	bool IsCompleteHandle() const { return Fields.Serial != InvalidValue && Fields.ReplicationSystemId != 0U; }
 
 	/** Static handles have ODD Id's */
-	bool IsStatic() const { return Id & StaticIdMask; }
+	bool IsStatic() const { return Fields.Static; }
 
 	/** Dynamic handles have EVEN Id's */
 	bool IsDynamic() const { return IsValid() && !IsStatic(); }
 
-	bool operator==(const FNetRefHandle& Other)const { return Id == Other.Id; }
-	bool operator<(const FNetRefHandle& Other)const { return Id < Other.Id; }
-	bool operator!=(const FNetRefHandle& Other)const { return Id != Other.Id; }
+	bool operator==(const FNetRefHandle& Other)const { return GetId() == Other.GetId(); }
+	bool operator<(const FNetRefHandle& Other)const { return GetId() < Other.GetId(); }
+	bool operator!=(const FNetRefHandle& Other)const { return GetId() != Other.GetId(); }
 
 	IRISCORE_API FString ToString() const;
 	IRISCORE_API FString ToCompactString() const;
 
-	static bool FullCompare(FNetRefHandle A, FNetRefHandle B) { return A.Value == B.Value; }
+	static bool FullCompare(FNetRefHandle A, FNetRefHandle B) 
+	{
+		return A.Fields.Static == B.Fields.Static &&
+			A.Fields.Serial == B.Fields.Serial &&
+			A.Fields.ReplicationSystemId == B.Fields.ReplicationSystemId;
+	}
 	
 private:
 	friend uint32 GetTypeHash(const FNetRefHandle& Handle);
 	friend Private::FNetRefHandleManager;
 
-	static constexpr uint64 StaticIdMask = 1;
-	static constexpr uint64 IdMask = (1ULL << IdBits) - 1;
-
-	union 
+	struct
 	{
-		struct
-		{
-			uint64 Id : IdBits;										// Id, lowest bit indicates if the handle is static or dynamic
-			uint64 ReplicationSystemId : ReplicationSystemIdBits;	// ReplicationSystemId, when running in pie, we track the owning instance
-		};
-		uint64 Value;
-	};
+		uint64 Static : StaticBits = InvalidValue;
+		uint64 Serial : SerialBits = InvalidValue;
+		uint64 ReplicationSystemId : ReplicationSystemIdBits = InvalidValue;
+	} Fields;
 };
 
 inline uint32 GetTypeHash(const FNetRefHandle& Handle)
diff --git a/Engine/Source/Runtime/Net/Iris/Public/Iris/ReplicationSystem/ReplicationSystem.h b/Engine/Source/Runtime/Net/Iris/Public/Iris/ReplicationSystem/ReplicationSystem.h
--- a/Engine/Source/Runtime/Net/Iris/Public/Iris/ReplicationSystem/ReplicationSystem.h
+++ b/Engine/Source/Runtime/Net/Iris/Public/Iris/ReplicationSystem/ReplicationSystem.h
@@ -170,6 +170,28 @@ public:
 
 		/** Configure the replication system as a proxy. */
 		UE::Net::EProxyType ProxyType = UE::Net::EProxyType::None;
+
+		/** 
+		 * Map a UObject pointer to the unique serial field in a NetRefHandle. 
+		 *
+		 * This delegate is called when assigning a NetRefHandle to an object being
+		 * registered for replication.
+		 *
+		 * Assumes there's a one-to-one mapping between UObjects registered
+		 * for replication and the serial field in NetRefHandle.
+		 * 
+		 * If left unset, Iris uses an internal incrementing counter to assign
+		 * a value to the NetRefHandle serial field.
+		 */
+		UE::Net::FObjectToNetRefHandleSerial ObjectToNetRefHandleSerial;
+
+		/** 
+		 * Map a unique serial field from a NetRefHandle to a UObject pointer.
+		 *
+		 * Assumes there's a one-to-one mapping between UObjects registered
+		 * for replication and the serial field in NetRefHandle.
+		 */
+		UE::Net::FNetRefHandleSerialToObject NetRefHandleSerialToObject;
 	};
 
 	/** @return The unique ID of the ReplicationSystem. */
diff --git a/Engine/Source/Runtime/Net/Iris/Public/Iris/ReplicationSystem/ReplicationSystemTypes.h b/Engine/Source/Runtime/Net/Iris/Public/Iris/ReplicationSystem/ReplicationSystemTypes.h
--- a/Engine/Source/Runtime/Net/Iris/Public/Iris/ReplicationSystem/ReplicationSystemTypes.h
+++ b/Engine/Source/Runtime/Net/Iris/Public/Iris/ReplicationSystem/ReplicationSystemTypes.h
@@ -47,6 +47,8 @@ enum class EDependentObjectSchedulingHint : uint8
 
 using FForwardNetRPCCallDelegate = TDelegate<void(UObject* RootObject,UObject* SubObject, UFunction* Function, void* Params)>;
 using FForwardNetRPCCallMulticastDelegate = TMulticastDelegate<typename FForwardNetRPCCallDelegate::TFuncType>;
+using FObjectToNetRefHandleSerial = TDelegate<uint64(const TObjectPtr<const UObject>& Object)>;
+using FNetRefHandleSerialToObject = TDelegate<TObjectPtr<UObject>(const uint64 Id)>;
 
 enum class ENetObjectAttachmentSendPolicyFlags : uint32
 {
diff --git a/Engine/Source/Runtime/Net/Iris/Public/Iris/Serialization/ObjectNetSerializer.h b/Engine/Source/Runtime/Net/Iris/Public/Iris/Serialization/ObjectNetSerializer.h
--- a/Engine/Source/Runtime/Net/Iris/Public/Iris/Serialization/ObjectNetSerializer.h
+++ b/Engine/Source/Runtime/Net/Iris/Public/Iris/Serialization/ObjectNetSerializer.h
@@ -58,7 +58,7 @@ UE_NET_DECLARE_SERIALIZER(FScriptInterfaceNetSerializer, IRISCORE_API);
 // in specific cases of forwarding net serializers.
 struct FObjectNetSerializerQuantizedReferenceStorage
 {
-	alignas(8) uint8 Storage[24];
+	alignas(8) uint8 Storage[32];
 };
 
 }
