Fix for bug where a NetRefHandle is configured to have a size of 0.

diff --git a/Engine/Source/Runtime/Net/Iris/Private/Iris/ReplicationSystem/NetRefHandleManager.cpp b/Engine/Source/Runtime/Net/Iris/Private/Iris/ReplicationSystem/NetRefHandleManager.cpp
--- a/Engine/Source/Runtime/Net/Iris/Private/Iris/ReplicationSystem/NetRefHandleManager.cpp
+++ b/Engine/Source/Runtime/Net/Iris/Private/Iris/ReplicationSystem/NetRefHandleManager.cpp
@@ -194,6 +194,11 @@ void FNetRefHandleManager::GrowNetChunkedArrayBuffers(FInternalNetRefIndex Inter
 
 uint64 FNetRefHandleManager::GetNextNetRefHandleSerial(uint64 HandleId) const
 {
+	if constexpr (FNetRefHandle::MaxSerial == 0)
+	{
+		return 0;
+	}
+
 	uint64 NextHandleId = (HandleId + 1) % FNetRefHandle::MaxSerial;
 	if (NextHandleId == 0)
 	{
diff --git a/Engine/Source/Runtime/Net/Iris/Public/Iris/ReplicationSystem/NetRefHandle.h b/Engine/Source/Runtime/Net/Iris/Public/Iris/ReplicationSystem/NetRefHandle.h
--- a/Engine/Source/Runtime/Net/Iris/Public/Iris/ReplicationSystem/NetRefHandle.h
+++ b/Engine/Source/Runtime/Net/Iris/Public/Iris/ReplicationSystem/NetRefHandle.h
@@ -39,6 +39,7 @@ private:
 	static constexpr uint64 IdBits = (StaticBits + SerialBits);
 	static constexpr uint64 ReplicationSystemIdBits = 10;
 
+	static_assert(SerialBits > 0, "The number of bits representing the FNetRefHandle serial must be greater than zero");
 	static_assert(IdBits <= 64, "The number of bits representing the FNetRefHandle id cannot be greater than 64-bits");
 
 public:
