Some transactional fixes for Iris.
- Add a bReplicateTransactionally flag to ReplicationSystem
- Pass bReplicateTransactionally from NetDriver to its ReplicationSystem
- Run QuantizeDirtyStateData in a transaction if necessary
- Call SetReplicateTransactionally in the MultiServerBeaconHost before the replication system is created so it picks up the flag

diff --git a/Engine/Source/Runtime/Engine/Private/NetDriver.cpp b/Engine/Source/Runtime/Engine/Private/NetDriver.cpp
--- a/Engine/Source/Runtime/Engine/Private/NetDriver.cpp
+++ b/Engine/Source/Runtime/Engine/Private/NetDriver.cpp
@@ -7490,6 +7490,7 @@ void UNetDriver::CreateReplicationSystem(bool bInitAsClient)
 		Params.bIsServer = !bInitAsClient;
 		Params.bAllowObjectReplication = !bInitAsClient;
 		Params.bUseRemoteObjectReferences = IsUsingRemoteObjectReferences();
+		Params.bReplicateTransactionally = ShouldReplicateTransactionally();
 		Params.ForwardNetRPCCallDelegate.BindUObject(this, &UNetDriver::ForwardRemoteFunction);
 		UE::Net::Private::ApplyReplicationSystemConfig(RepSystemConfig, Params);
 		Params.bAllowExecuteRPC = IsExecuteRPCFunctionsEnabled();
diff --git a/Engine/Source/Runtime/Engine/Private/PlayerController.cpp b/Engine/Source/Runtime/Engine/Private/PlayerController.cpp
--- a/Engine/Source/Runtime/Engine/Private/PlayerController.cpp
+++ b/Engine/Source/Runtime/Engine/Private/PlayerController.cpp
@@ -5112,6 +5112,9 @@ void APlayerController::PostMigrate(const FUObjectMigrationContext& MigrationCon
 				ensure(PrevNetPlayerIndex == NetPlayerIndex);
 				ensure(PrevHandshakeTransactionId == GetClientHandshakeId());
 
+				// Update Iris with the new Player/NetConnection
+				UpdateOwningNetConnection();
+
 				UE_LOG(LogPlayerController, Verbose, TEXT("[%s] APlayerController::PostMigrate [%s]: Binding %s to connection %s"), 
 					   *FRemoteServerId::GetLocalServerId().ToString(), *FRemoteObjectId(this).ToString(), 
 					   *GetName(), *NetConnection->GetName());
@@ -5154,11 +5157,12 @@ void APlayerController::PostMigrate(const FUObjectMigrationContext& MigrationCon
 				PrevNetConnection->PlayerController->SetRole(ROLE_Authority);
 				PrevNetConnection->PlayerController->SetReplicates(true);
 				PrevNetConnection->PlayerController->SetAutonomousProxy(true);
+				// Set the NetConnection before FinishSpawning so that the AActor::UpdateOwningNetConnection call properly updates the Iris conection owner.
+				PrevNetConnection->PlayerController->NetConnection = PrevNetConnection;
 				PrevNetConnection->PlayerController->FinishSpawning(FTransform());
 				PrevNetConnection->PlayerController->CachedConnectionPlayerId = CachedConnectionPlayerId;
 				PrevNetConnection->PlayerController->NetPlayerIndex = NetPlayerIndex;
 				PrevNetConnection->PlayerController->SetClientHandshakeId(PrevNetConnection->GetClientHandshakeId());
-				PrevNetConnection->PlayerController->NetConnection = PrevNetConnection;
 
 				PrevNetConnection->OwningActor = PrevNetConnection->PlayerController;
 					
diff --git a/Engine/Source/Runtime/Net/Iris/Private/Iris/ReplicationSystem/ReplicationSystem.cpp b/Engine/Source/Runtime/Net/Iris/Private/Iris/ReplicationSystem/ReplicationSystem.cpp
--- a/Engine/Source/Runtime/Net/Iris/Private/Iris/ReplicationSystem/ReplicationSystem.cpp
+++ b/Engine/Source/Runtime/Net/Iris/Private/Iris/ReplicationSystem/ReplicationSystem.cpp
@@ -39,6 +39,7 @@
 
 #include "HAL/IConsoleManager.h"
 #include "ProfilingDebugging/CsvProfiler.h"
+#include "UObject/RemoteExecutor.h"
 #include "UObject/UObjectGlobals.h"
 
 #include UE_INLINE_GENERATED_CPP_BY_NAME(ReplicationSystem)
@@ -559,6 +560,24 @@ public:
 		UE_NET_TRACE_FRAME_STATSCOUNTER(ReplicationSystemId, ReplicationSystem.QuantizedObjectCount, QuantizedObjectCount, ENetTraceVerbosity::Trace);
 	}
 
+	void ExecuteQuantizeDirtyStateData(bool bTransact)
+	{
+#if UE_WITH_REMOTE_OBJECT_HANDLE && UE_AUTORTFM
+		if (bTransact)
+		{
+			static FName TransactionalWorkName("Iris FReplicationSystemImpl::QuantizeDirtyStateData");
+			UE::RemoteExecutor::ExecuteTransactional(TransactionalWorkName, [this]()
+				{
+					QuantizeDirtyStateData();
+				});
+		}
+		else
+#endif
+		{
+			QuantizeDirtyStateData();
+		}
+	}
+
 	void ResetObjectStateDirtiness()
 	{
 		IRIS_CSV_PROFILER_SCOPE(Iris, ReplicationSystem_ResetObjectStateDirtiness);
@@ -835,6 +854,7 @@ UReplicationSystem::UReplicationSystem()
 , bIsServer(0)
 , bAllowObjectReplication(0)
 , bDoCollectGarbage(0)
+, bReplicateTransactionally(0)
 {
 }
 
@@ -846,6 +866,7 @@ void UReplicationSystem::Init(uint32 InId, const FReplicationSystemParams& Param
 	bAllowExecuteRPC = Params.bAllowExecuteRPC;
 	bAllowExecuteRepNotify = Params.bAllowExecuteRepNotifies;
 	bAlwaysSwapRolesOnReplication = Params.bAlwaysSwapRolesOnReplication;
+	bReplicateTransactionally = Params.bReplicateTransactionally;
 	ProxyType = Params.ProxyType;
 
 	ReplicationBridge = Params.ReplicationBridge;
@@ -955,7 +976,7 @@ void UReplicationSystem::NetUpdate(float DeltaSeconds)
 			Impl->UpdateConditionals();
 
 			// Quantize dirty state data. We need this to happen before both filtering and prioritization
-			Impl->QuantizeDirtyStateData();
+			Impl->ExecuteQuantizeDirtyStateData(bReplicateTransactionally);
 
 			// We must process all attachments to objects going out of scope before we update the scope
 			Impl->ProcessNetObjectAttachmentSendQueue(FNetBlobManager::EProcessMode::ProcessObjectsGoingOutOfScope);
diff --git a/Engine/Source/Runtime/Net/Iris/Public/Iris/ReplicationSystem/ReplicationSystem.h b/Engine/Source/Runtime/Net/Iris/Public/Iris/ReplicationSystem/ReplicationSystem.h
--- a/Engine/Source/Runtime/Net/Iris/Public/Iris/ReplicationSystem/ReplicationSystem.h
+++ b/Engine/Source/Runtime/Net/Iris/Public/Iris/ReplicationSystem/ReplicationSystem.h
@@ -138,6 +138,12 @@ public:
 		/** When using remote objects with UE_WITH_REMOTE_OBJECT_HANDLE, determines whether or not TObjectPtr or TWeakObjectPtr properties are serialized as FRemoteObjectReferences */
 		bool bUseRemoteObjectReferences = false;
 
+		/**
+		 * If true, when using remote objects with UE_WITH_REMOTE_OBJECT_HANDLE, runs code that may dereference remote UObjects inside of a transaction.
+		 * If the replication system is being used by a NetDriver, the NetDriver should set this flag if the NetDriver is also replicating transactionally.
+		 */
+		bool bReplicateTransactionally = false;
+
 		/** 
 		*   When true, allow subsystems to run parallel workloads, such as the PollAndCopy step running several asynchronous tasks to speed up game thread execution time.
 		*   Only supported when bIsServer = true and bAllowObjectReplication = true
@@ -834,6 +840,7 @@ private:
 	uint32 bAllowExecuteRepNotify : 1;
 	uint32 bAlwaysSwapRolesOnReplication : 1;
 	uint32 bDoCollectGarbage : 1;
+	uint32 bReplicateTransactionally : 1;
 
 	UE::Net::EProxyType ProxyType = UE::Net::EProxyType::None;
 };
