Iris
* Added Bridge::IsRootObject function

#rnx

diff --git a/Engine/Source/Runtime/Net/Iris/Private/Iris/ReplicationSystem/NetRefHandleManager.h b/Engine/Source/Runtime/Net/Iris/Private/Iris/ReplicationSystem/NetRefHandleManager.h
--- a/Engine/Source/Runtime/Net/Iris/Private/Iris/ReplicationSystem/NetRefHandleManager.h
+++ b/Engine/Source/Runtime/Net/Iris/Private/Iris/ReplicationSystem/NetRefHandleManager.h
@@ -237,6 +237,9 @@ public:
 	/** Returns true if the handle is for a remote replicated object. */
 	inline bool IsRemoteNetRefHandle(FNetRefHandle Handle) const;
 
+	/** Returns true if the handle is associated with a root object */
+	inline bool IsRootObject(FNetRefHandle Handle) const;
+
 	/** Extract a full handle from an incomplete one consisting of only an index. */
 	inline FNetRefHandle GetCompleteNetRefHandle(FNetRefHandle IncompleteHandle) const;
 	
@@ -284,6 +287,11 @@ public:
 
 	bool SetSubObjectNetCondition(FInternalNetRefIndex SubObjectInternalIndex, FLifeTimeConditionStorage SubObjectCondition, bool& bOutWasModified);
 
+	bool IsRootObject(FInternalNetRefIndex InternalIndex) const
+	{
+		return SubObjectInternalIndices.IsBitSet(InternalIndex) == false;
+	}
+
 	bool IsSubObject(FInternalNetRefIndex InternalIndex) const
 	{
 		return SubObjectInternalIndices.GetBit(InternalIndex);
@@ -649,6 +657,12 @@ bool FNetRefHandleManager::IsRemoteNetRefHandle(FNetRefHandle Handle) const
 	return false;
 }
 
+bool FNetRefHandleManager::IsRootObject(FNetRefHandle Handle) const
+{
+	const uint32 InternalIndex = GetInternalIndex(Handle);
+	return (InternalIndex != InvalidInternalIndex) ? IsRootObject(InternalIndex) : false;
+}
+
 uint32 FNetRefHandleManager::GetOriginalDestroyedStartupObjectIndex(FInternalNetRefIndex InternalIndex) const
 {
 	const uint32* FoundOriginalInternalIndex = DestroyedStartupObject.Find(InternalIndex);
diff --git a/Engine/Source/Runtime/Net/Iris/Private/Iris/ReplicationSystem/ObjectReplicationBridge.cpp b/Engine/Source/Runtime/Net/Iris/Private/Iris/ReplicationSystem/ObjectReplicationBridge.cpp
--- a/Engine/Source/Runtime/Net/Iris/Private/Iris/ReplicationSystem/ObjectReplicationBridge.cpp
+++ b/Engine/Source/Runtime/Net/Iris/Private/Iris/ReplicationSystem/ObjectReplicationBridge.cpp
@@ -990,6 +990,11 @@ void UObjectReplicationBridge::SetSubObjectNetCondition(FNetRefHandle SubObjectR
 	}
 }
 
+bool UObjectReplicationBridge::IsRootObject(FNetRefHandle RootObjectHandle) const
+{
+	return GetReplicationSystem()->GetReplicationSystemInternal()->GetNetRefHandleManager().IsRootObject(RootObjectHandle);
+}
+
 UE::Net::FNetRefHandle UObjectReplicationBridge::GetRootObjectOfAnyObject(FNetRefHandle NetRefHandle) const
 {
 	return GetReplicationSystem()->GetReplicationSystemInternal()->GetNetRefHandleManager().GetRootObjectOfAnyObject(NetRefHandle);
diff --git a/Engine/Source/Runtime/Net/Iris/Public/Iris/ReplicationSystem/ObjectReplicationBridge.h b/Engine/Source/Runtime/Net/Iris/Public/Iris/ReplicationSystem/ObjectReplicationBridge.h
--- a/Engine/Source/Runtime/Net/Iris/Public/Iris/ReplicationSystem/ObjectReplicationBridge.h
+++ b/Engine/Source/Runtime/Net/Iris/Public/Iris/ReplicationSystem/ObjectReplicationBridge.h
@@ -148,7 +148,11 @@ public:
 	 */
 	IRISCORE_API void SetSubObjectNetCondition(FNetRefHandle SubObjectHandle, ELifetimeCondition Condition);
 
-	/** Returns the root object handle associated with any object. 
+	/** Returns true if the handle is valid and associated with a root object */
+	IRISCORE_API bool IsRootObject(FNetRefHandle RootObjectHandle) const;
+
+	/** 
+	 * Returns the root object handle associated with any object. 
 	 * If the passed handle is a subobject it will return its root.  
 	 * If the passed handle is a rootobject it will return the same handle.
 	 */
