Iris Parallelization - Added thread safety to NetBlob, NetTokenStore, NameTokenStore, StringTokenStore and StructNetTokenDataStore
Added UE_SUPPORT_PARALLEL_IRIS define which includes multithreading logic only on Server and Editor builds by default. Client only builds do not need the mutexes included as they will not run the replication logic in parallel.
Added ExcludeParallelIrisSupport to TargetRules.cs so that licensees can remove these mutexes if they intend to run the server in single threaded and want a few nanoseconds of perf back.

[RN] Iris Parallelization - Added thread safety to NetBlob and TokenStore. Can be removed from compilation using ExcludeParallelIrisSupport
#tests Built and tested against Client + Server Cooked Data and Editor in PIE

diff --git a/Engine/Plugins/Runtime/ReplicationSystemTestPlugin/Source/Private/Tests/ReplicationSystem/NetBlob/TestReliableNetBlobQueue.cpp b/Engine/Plugins/Runtime/ReplicationSystemTestPlugin/Source/Private/Tests/ReplicationSystem/NetBlob/TestReliableNetBlobQueue.cpp
--- a/Engine/Plugins/Runtime/ReplicationSystemTestPlugin/Source/Private/Tests/ReplicationSystem/NetBlob/TestReliableNetBlobQueue.cpp
+++ b/Engine/Plugins/Runtime/ReplicationSystemTestPlugin/Source/Private/Tests/ReplicationSystem/NetBlob/TestReliableNetBlobQueue.cpp
@@ -4,6 +4,7 @@
 #include "MockNetBlob.h"
 #include "Iris/ReplicationSystem/NetBlob/ReliableNetBlobQueue.h"
 #include "Iris/ReplicationSystem/NetBlob/PartialNetObjectAttachmentHandler.h"
+#include "AutoRTFM.h"
 
 namespace UE::Net::Private
 {
@@ -143,4 +144,77 @@ TRefCountPtr<FNetObjectAttachment> FTestReliableNetBlobQueue::CreateReliableNetO
 	}
 }
 
+#if UE_AUTORTFM
+UE_NET_TEST_FIXTURE(FTestReliableNetBlobQueue, RefCountingNetBlobsInsideAutoRTFMTransactions)
+{
+	FReplicationSystemTestClient* Client = CreateClient();
+	RegisterNetBlobHandlers(Client);
+
+	UReplicatedTestObject* ServerObject = Server->CreateObject(0, 0);
+
+	Server->UpdateAndSend({ Client });
+
+	UE_NET_ASSERT_NE(Client->GetReplicationBridge()->GetReplicatedObject(ServerObject->NetRefHandle), nullptr);
+
+	// Perform the test with Committing or Aborting the AutoRTFM transaction
+	for (bool bCommitTransaction : {false, true})
+	{
+		// Start with a clean slate for each test.
+		ClientMockNetObjectAttachmentHandler->ResetFunctionCallCounts();
+
+		// Create attachment with a single part count
+		constexpr uint32 BlobPartCount = 1;
+		{
+			const TRefCountPtr<FNetObjectAttachment>& Attachment = CreateReliableNetObjectAttachmentWithPartCount(BlobPartCount);
+			FNetObjectReference AttachmentTarget = FObjectReferenceCache::MakeNetObjectReference(ServerObject->NetRefHandle);
+			Server->GetReplicationSystem()->QueueNetObjectAttachment(Client->ConnectionIdOnServer, AttachmentTarget, Attachment);
+
+			// Now take take a reference to this attachment inside an AutoRTFM scope. Whether it succeeds or not, we should
+			// end up with the same RefCount afterwards
+			int32 InitialAttachmentRefCount = const_cast<TRefCountPtr<FNetObjectAttachment>&>(Attachment).GetRefCount();
+			AutoRTFM::Transact([&]
+			{
+				TRefCountPtr<FNetObjectAttachment> NewRef(Attachment);
+			
+				UE_NET_ASSERT_GT(NewRef.GetRefCount(), 0U);
+			
+				if(!bCommitTransaction)
+				{
+					AutoRTFM::AbortTransaction();
+				}
+			});
+			int32 FinalAttachmentRefCount = const_cast<TRefCountPtr<FNetObjectAttachment>&>(Attachment).GetRefCount();
+
+			UE_NET_ASSERT_EQ(InitialAttachmentRefCount, FinalAttachmentRefCount);
+		}
+
+		// Deliver as much as possible. We stop when no more packets are written.
+		const auto& ConnectionInfo = Server->GetConnectionInfo(Client->ConnectionIdOnServer);
+		SIZE_T PacketCount = ConnectionInfo.WrittenPackets.Count();
+		for (SIZE_T TryIt = 0, TryEndIt = BlobPartCount; TryIt != TryEndIt; ++TryIt)
+		{
+			Server->UpdateAndSend({ Client });
+
+			const SIZE_T NewPacketCount = ConnectionInfo.WrittenPackets.Count();
+			if (NewPacketCount == PacketCount)
+			{
+				break;
+			}
+
+			PacketCount = NewPacketCount;
+		}
+
+		// Now do the receive logic on the client. 
+		for (SIZE_T PacketIt = 0; PacketIt != PacketCount; ++PacketIt)
+		{
+			Server->DeliverTo(Client, DeliverPacket);
+		}
+
+		// Make sure the blob was received on the client and there are no attachments waiting to be sent
+		UE_NET_ASSERT_EQ(ClientMockNetObjectAttachmentHandler->GetFunctionCallCounts().OnNetBlobReceived, 1U);
+		UE_NET_ASSERT_EQ(Server->GetReplicationSystem()->GetReplicationSystemInternal()->GetNetBlobManager().HasAnyUnprocessedReliableAttachments(), false);
+	}
+}
+#endif // UE_AUTORTFM
+
 }
diff --git a/Engine/Source/Programs/UnrealBuildTool/Configuration/ReadOnlyTargetRules.cs b/Engine/Source/Programs/UnrealBuildTool/Configuration/ReadOnlyTargetRules.cs
--- a/Engine/Source/Programs/UnrealBuildTool/Configuration/ReadOnlyTargetRules.cs
+++ b/Engine/Source/Programs/UnrealBuildTool/Configuration/ReadOnlyTargetRules.cs
@@ -227,6 +227,8 @@ namespace UnrealBuildTool
 
 		public bool bWithPushModel => Inner.bWithPushModel;
 
+		public bool bExcludeParallelIrisSupport => Inner.bExcludeParallelIrisSupport;
+
 		public bool bCompileWithStatsWithoutEngine => Inner.bCompileWithStatsWithoutEngine;
 
 		public bool bCompileWithPluginSupport => Inner.bCompileWithPluginSupport;
diff --git a/Engine/Source/Programs/UnrealBuildTool/Configuration/TargetRules.cs b/Engine/Source/Programs/UnrealBuildTool/Configuration/TargetRules.cs
--- a/Engine/Source/Programs/UnrealBuildTool/Configuration/TargetRules.cs
+++ b/Engine/Source/Programs/UnrealBuildTool/Configuration/TargetRules.cs
@@ -1441,6 +1441,14 @@ namespace UnrealBuildTool
 		}
 		private bool? bWithPushModelOverride;
 
+		/// <summary>
+		/// When enabled, certain thread safety mechanisms will be removed from Iris code
+		/// this may result in a small cpu perf gain if running a single-threaded server
+		/// </summary>
+		[RequiresUniqueBuildEnvironment]
+		[CommandLine("-ExcludeParallelIrisSupport", Value = "true")]
+		public bool bExcludeParallelIrisSupport { get; set; } = false;
+
 		/// <summary>
 		/// Whether to include stats support even without the engine.
 		/// </summary>
diff --git a/Engine/Source/Runtime/Engine/Public/Net/Iris/ReplicationSystem/StructNetTokenDataStoreHelper.h b/Engine/Source/Runtime/Engine/Public/Net/Iris/ReplicationSystem/StructNetTokenDataStoreHelper.h
--- a/Engine/Source/Runtime/Engine/Public/Net/Iris/ReplicationSystem/StructNetTokenDataStoreHelper.h
+++ b/Engine/Source/Runtime/Engine/Public/Net/Iris/ReplicationSystem/StructNetTokenDataStoreHelper.h
@@ -60,6 +60,8 @@ public:
 				return false;
 			}
 
+			UE_IRIS_PARALLEL_EXPR(FScopedNetTokenStoreAccess NetTokenStoreAccess(NetTokenStore));
+
 			UPackageMapClient* PackageMapClient = Cast<UPackageMapClient>(Map);
 			if (PackageMapClient && PackageMapClient->GetConnection() && PackageMapClient->GetConnection()->IsReplay())
 			{
@@ -80,6 +82,8 @@ public:
 				return false;
 			}
 
+			UE_IRIS_PARALLEL_EXPR(FScopedNetTokenStoreAccess NetTokenStoreAccess(NetTokenStore));
+
 			UPackageMapClient* PackageMapClient = Cast<UPackageMapClient>(Map);
 			if (PackageMapClient && PackageMapClient->GetConnection() && PackageMapClient->GetConnection()->IsReplay())
 			{
diff --git a/Engine/Source/Runtime/Net/Core/Private/Net/Core/Trace/PersistentNetDebugNames.cpp b/Engine/Source/Runtime/Net/Core/Private/Net/Core/Trace/PersistentNetDebugNames.cpp
--- a/Engine/Source/Runtime/Net/Core/Private/Net/Core/Trace/PersistentNetDebugNames.cpp
+++ b/Engine/Source/Runtime/Net/Core/Private/Net/Core/Trace/PersistentNetDebugNames.cpp
@@ -8,6 +8,12 @@
 #include "Misc/MemStack.h"
 #include "Misc/ScopeLock.h"
 
+
+namespace UE::Net::Private
+{
+	static FNetDebugName const* ReliabilityNetDebugNames[2];
+}
+
 namespace UE::Net
 {
 
@@ -46,6 +52,11 @@ FPersistentNetDebugNames& FPersistentNetDebugNames::Get()
 	{
 		FCoreDelegates::OnExit.AddStatic(&TearDown);
 		Singleton.bIsInitialized = true;
+
+#if UE_NET_TRACE_ENABLED
+		Private::ReliabilityNetDebugNames[0] = UE::Net::CreatePersistentNetDebugName(TEXT("Unreliable"));
+		Private::ReliabilityNetDebugNames[1] = UE::Net::CreatePersistentNetDebugName(TEXT("Reliable"));
+#endif
 	}
 
 	return Singleton;
@@ -107,6 +118,11 @@ bool FPersistentNetDebugNames::CreatePersistentNetDebugName(const TCHAR* Name, u
 	return true;
 }
 
+FNetDebugName const* GetReliabilityNetDebugName(bool bIsReliable)
+{
+	return Private::ReliabilityNetDebugNames[bIsReliable ? 1 : 0];
+}
+
 const FNetDebugName* CreatePersistentNetDebugName(const TCHAR* Name, uint32 NameLen)
 {
 	const FNetDebugName* DebugName;
diff --git a/Engine/Source/Runtime/Net/Core/Public/Net/Core/Trace/NetDebugName.h b/Engine/Source/Runtime/Net/Core/Public/Net/Core/Trace/NetDebugName.h
--- a/Engine/Source/Runtime/Net/Core/Public/Net/Core/Trace/NetDebugName.h
+++ b/Engine/Source/Runtime/Net/Core/Public/Net/Core/Trace/NetDebugName.h
@@ -25,6 +25,8 @@ inline const FNetDebugName* CreatePersistentNetDebugName(const TCHAR* Name) { re
 
 NETCORE_API void ResetPersistentNetDebugNameIds();
 
+NETCORE_API FNetDebugName const* GetReliabilityNetDebugName(bool bIsReliable);
+
 }
 
 inline const TCHAR* ToCStr(const UE::Net::FNetDebugName* DebugName)
diff --git a/Engine/Source/Runtime/Net/Iris/IrisCore.Build.cs b/Engine/Source/Runtime/Net/Iris/IrisCore.Build.cs
--- a/Engine/Source/Runtime/Net/Iris/IrisCore.Build.cs
+++ b/Engine/Source/Runtime/Net/Iris/IrisCore.Build.cs
@@ -19,6 +19,12 @@ namespace UnrealBuildTool.Rules
 			// To be removed once UE_WITH_IRIS has been removed in code.
 			PublicDefinitions.Add("UE_WITH_IRIS=1");
 
+			//Only compile in support for Iris Parallel thread safety if not Client,Program and not disabled on the build command line 
+			if (!Target.bExcludeParallelIrisSupport && !(Target.Type == TargetType.Client || Target.Type == TargetType.Program))
+			{
+				PublicDefinitions.Add("UE_SUPPORT_PARALLEL_IRIS=1");
+			}
+
 			CppCompileWarningSettings.UnsafeTypeCastWarningLevel = WarningLevel.Error;
 		}
 	}
diff --git a/Engine/Source/Runtime/Net/Iris/Private/Iris/ReplicationSystem/NameTokenStore.cpp b/Engine/Source/Runtime/Net/Iris/Private/Iris/ReplicationSystem/NameTokenStore.cpp
--- a/Engine/Source/Runtime/Net/Iris/Private/Iris/ReplicationSystem/NameTokenStore.cpp
+++ b/Engine/Source/Runtime/Net/Iris/Private/Iris/ReplicationSystem/NameTokenStore.cpp
@@ -48,14 +48,18 @@ FNetToken FNameTokenStore::GetOrCreateToken(FName Name)
 
 FNetTokenDataStore::FNetTokenStoreKey FNameTokenStore::GetOrCreateTokenStoreKey(FName Name)
 {
-	if (const FNetTokenStoreKey* ExistingKey = FNameToKey.Find(Name))
 	{
-		return *ExistingKey;
+		UE_IRIS_PARALLEL_EXPR(UE::TRWScopeLock FNameReadLock(NameTokenStoreMutex, SLT_ReadOnly));
+		if (const FNetTokenStoreKey* ExistingKey = FNameToKey.Find(Name))
+		{
+			return *ExistingKey;
+		}
 	}
 
 	const FNetTokenStoreKey NewKey = GetNextNetTokenStoreKey();
 	if (NewKey.IsValid())
 	{
+		UE_IRIS_PARALLEL_EXPR(UE::TRWScopeLock FNameWriteLock(NameTokenStoreMutex, SLT_Write));
 		StoredFNames.Add(Name);
 		FNameToKey.Add(Name, NewKey);
 
diff --git a/Engine/Source/Runtime/Net/Iris/Private/Iris/ReplicationSystem/NetBlob/NetBlob.cpp b/Engine/Source/Runtime/Net/Iris/Private/Iris/ReplicationSystem/NetBlob/NetBlob.cpp
--- a/Engine/Source/Runtime/Net/Iris/Private/Iris/ReplicationSystem/NetBlob/NetBlob.cpp
+++ b/Engine/Source/Runtime/Net/Iris/Private/Iris/ReplicationSystem/NetBlob/NetBlob.cpp
@@ -124,10 +124,13 @@ void FNetBlob::DeserializeBlob(FNetSerializationContext& Context)
 
 void FNetBlob::Release() const
 {
-	if (--RefCount == 0)
+	UE_AUTORTFM_ONCOMMIT(this)
 	{
-		delete this;
-	}
+		if (--RefCount == 0)
+		{
+			delete this;
+		}
+	};
 }
 
 TArrayView<const FNetObjectReference> FNetBlob::GetNetObjectReferenceExports() const
diff --git a/Engine/Source/Runtime/Net/Iris/Private/Iris/ReplicationSystem/NetBlob/PartialNetBlob.cpp b/Engine/Source/Runtime/Net/Iris/Private/Iris/ReplicationSystem/NetBlob/PartialNetBlob.cpp
--- a/Engine/Source/Runtime/Net/Iris/Private/Iris/ReplicationSystem/NetBlob/PartialNetBlob.cpp
+++ b/Engine/Source/Runtime/Net/Iris/Private/Iris/ReplicationSystem/NetBlob/PartialNetBlob.cpp
@@ -17,7 +17,7 @@ namespace UE::Net::Private
 	 * Each split NetBlob will reserve unique sequence numbers for all of its parts. If there are multiple ReplicationSystems sharing the atomic adds relatively little growing of the sequence numbers.
 	 * Splitting blobs is a special case mainly used to enable replicating very large objects.
 	 */
-	static uint32 PartialNetBlobGlobalSequenceNumber;
+	static std::atomic<uint32> PartialNetBlobGlobalSequenceNumber;
 }
 
 namespace UE::Net
@@ -271,8 +271,15 @@ void FPartialNetBlob::SplitPayload(const FPartialNetBlob::FPayloadSplitParams& S
 	OutPartialBlobs.Reserve(OutPartialBlobs.Num() + int32(PartialBlobCount));
 	
 	// Reserve sequence numbers for all parts. 
-	uint32 SequenceNumber = Private::PartialNetBlobGlobalSequenceNumber;
-	Private::PartialNetBlobGlobalSequenceNumber += PartialBlobCount;
+	uint32 SequenceNumber = 0;
+	// We do not currently run SplitPayload in a closed AutoRTFM scope, but it is acceptable for the sequence number to
+	// increase and never roll back, so this makes it safe to access this atomic in a closed scope in the future
+	UE_AUTORTFM_OPEN
+	{
+		SequenceNumber = Private::PartialNetBlobGlobalSequenceNumber;
+		Private::PartialNetBlobGlobalSequenceNumber += PartialBlobCount;
+	};
+	
 
 	uint32 PayloadBitOffset = 0U;
 
diff --git a/Engine/Source/Runtime/Net/Iris/Private/Iris/ReplicationSystem/NetBlob/ReliableNetBlobQueue.cpp b/Engine/Source/Runtime/Net/Iris/Private/Iris/ReplicationSystem/NetBlob/ReliableNetBlobQueue.cpp
--- a/Engine/Source/Runtime/Net/Iris/Private/Iris/ReplicationSystem/NetBlob/ReliableNetBlobQueue.cpp
+++ b/Engine/Source/Runtime/Net/Iris/Private/Iris/ReplicationSystem/NetBlob/ReliableNetBlobQueue.cpp
@@ -19,8 +19,6 @@ namespace UE::Net::Private
 static const FName NetError_ReliableQueueFull("Reliable attachment queue full");
 static const FName NetError_InvalidSequence("Invalid sequence number");
 
-static FNetDebugName const* ReliabilityNetDebugNames[2];
-
 // ReliableNetBlobQueue
 FReliableNetBlobQueue::FReliableNetBlobQueue()
 : Sent{}
@@ -29,10 +27,6 @@ FReliableNetBlobQueue::FReliableNetBlobQueue()
 , LastSeq(0)
 , UnsentBlobCount(0)
 {
-#if UE_NET_TRACE_ENABLED
-	Private::ReliabilityNetDebugNames[0] = CreatePersistentNetDebugName(TEXT("Unreliable"));
-	Private::ReliabilityNetDebugNames[1] = CreatePersistentNetDebugName(TEXT("Reliable"));
-#endif
 }
 
 FReliableNetBlobQueue::~FReliableNetBlobQueue()
@@ -127,7 +121,7 @@ uint32 FReliableNetBlobQueue::SerializeInternal(FNetSerializationContext& Contex
 		const bool bHasData = Attachment.GetRefCount() > 0;
 		if (Writer->WriteBool(bHasData))
 		{
-			UE_NET_TRACE_NAMED_DYNAMIC_NAME_SCOPE(ReliabilityScope, Private::ReliabilityNetDebugNames[Attachment->IsReliable()], *Writer, Context.GetTraceCollector(), ENetTraceVerbosity::Verbose);
+			UE_NET_TRACE_NAMED_DYNAMIC_NAME_SCOPE(ReliabilityScope, GetReliabilityNetDebugName(Attachment->IsReliable()), *Writer, Context.GetTraceCollector(), ENetTraceVerbosity::Verbose);
 
 			// If we have exports, append them, if attachment is rolled back we will roll back any appended exports as well.
 			FNetExportContext* ExportContext = Attachment->HasExports() ? Context.GetExportContext() : nullptr;
@@ -300,7 +294,7 @@ uint32 FReliableNetBlobQueue::DeserializeInternal(FNetSerializationContext& Cont
 		TRefCountPtr<FNetBlob> Blob;
 		if (const bool bHasData = Reader->ReadBool())
 		{
-			UE_NET_TRACE_NAMED_DYNAMIC_NAME_SCOPE(ReliabilityScope, Private::ReliabilityNetDebugNames[0], *Reader, Context.GetTraceCollector(), ENetTraceVerbosity::Verbose);
+			UE_NET_TRACE_NAMED_DYNAMIC_NAME_SCOPE(ReliabilityScope, GetReliabilityNetDebugName(false), *Reader, Context.GetTraceCollector(), ENetTraceVerbosity::Verbose);
 
 			FNetBlobCreationInfo CreationInfo;
 			FNetBlob::DeserializeCreationInfo(Context, CreationInfo);
@@ -322,7 +316,7 @@ uint32 FReliableNetBlobQueue::DeserializeInternal(FNetSerializationContext& Cont
 				Blob->Deserialize(Context);
 			}
 
-			UE_NET_TRACE_SET_SCOPE_NAME(ReliabilityScope, Private::ReliabilityNetDebugNames[Blob->IsReliable()]);
+			UE_NET_TRACE_SET_SCOPE_NAME(ReliabilityScope, GetReliabilityNetDebugName(Blob->IsReliable()));
 		}
 		
 		bHasMoreBlobs = Reader->ReadBool();
diff --git a/Engine/Source/Runtime/Net/Iris/Private/Iris/ReplicationSystem/NetTokenStore.cpp b/Engine/Source/Runtime/Net/Iris/Private/Iris/ReplicationSystem/NetTokenStore.cpp
--- a/Engine/Source/Runtime/Net/Iris/Private/Iris/ReplicationSystem/NetTokenStore.cpp
+++ b/Engine/Source/Runtime/Net/Iris/Private/Iris/ReplicationSystem/NetTokenStore.cpp
@@ -70,6 +70,8 @@ FNetTokenDataStore::~FNetTokenDataStore()
 
 FNetTokenDataStore::FNetTokenStoreKey FNetTokenDataStore::GetTokenKey(FNetToken Token, const FNetTokenStoreState& TokenStoreState) const
 {
+	UE_IRIS_PARALLEL_EXPR(FScopedNetTokenStoreAccess ScopeAccess(&TokenStore));
+
 	if (Token.GetTypeId() == GetTypeId())
 	{
 		const TArray<FNetTokenStoreKey>& TokenStoreKeysForType = TokenStoreState.TokenInfoArray[GetTypeId()];
@@ -86,6 +88,8 @@ FNetTokenDataStore::FNetTokenStoreKey FNetTokenDataStore::GetTokenKey(FNetToken
 
 FNetToken FNetTokenDataStore::CreateAndStoreTokenForKey(FNetTokenStoreKey Key)
 {
+	UE_IRIS_PARALLEL_EXPR(FScopedNetTokenStoreAccess ScopeAccess(&TokenStore));
+
 	FNetTokenStoreState& LocalNetTokenStoreState = *TokenStore.GetLocalNetTokenStoreState();
 
 	// The LocalNetTokenStoreState always maps locally assigned tokens -> Key (Index)
@@ -109,12 +113,15 @@ FNetToken FNetTokenDataStore::CreateAndStoreTokenForKey(FNetTokenStoreKey Key)
 
 void FNetTokenDataStore::StoreTokenForKey(FNetTokenStoreKey Key, FNetToken NetToken)
 {
+	UE_IRIS_PARALLEL_EXPR(FScopedNetTokenStoreAccess ScopeAccess(&TokenStore));
+
 	// The stored tokens array contains the current NetToken associated with the Key (Index), but it can be updated to use an authoritative token instead.
 	StoredTokens[Key.GetKeyIndex()] = NetToken;
 }
 
 FNetToken FNetTokenDataStore::GetNetTokenFromKey(FNetTokenStoreKey Key) const
 {
+	UE_IRIS_PARALLEL_EXPR(FScopedNetTokenStoreAccess ScopeAccess(&TokenStore));
 	return StoredTokens[Key.GetKeyIndex()];
 }
 
@@ -373,6 +380,8 @@ void FNetTokenStore::WriteTokenData(FNetSerializationContext& Context, const FNe
 {
 	if (NetToken.IsValid())
 	{
+		UE_IRIS_PARALLEL_EXPR(FScopedNetTokenStoreAccess ScopeAccess(this));
+
 		FNetBitStreamWriter* Writer = Context.GetBitStreamWriter();
 
 		const FNetTokenStoreKey& TokenKey = LocalNetTokenStoreState->TokenInfoArray[NetToken.GetTypeId()][NetToken.GetIndex()];
@@ -386,6 +395,8 @@ void FNetTokenStore::WriteTokenData(FArchive& Ar, const FNetToken NetToken, UPac
 {
 	if (NetToken.IsValid())
 	{
+		UE_IRIS_PARALLEL_EXPR(FScopedNetTokenStoreAccess ScopeAccess(this));
+
 		const FNetToken::FTypeId TokenTypeId = NetToken.GetTypeId();
 
 		// We cannot write token data for bad tokens
@@ -404,6 +415,7 @@ void FNetTokenStore::WriteTokenData(FArchive& Ar, const FNetToken NetToken, UPac
 
 bool FNetTokenStore::ValidateAndStoreNetTokenData(FNetTokenDataStore& DataStore, FNetTokenStoreState& RemoteNetTokenStoreState, const FNetToken NetToken, const FNetTokenStoreKey StoreKey)
 {
+	UE_IRIS_PARALLEL_EXPR(FScopedNetTokenStoreAccess ScopeAccess(this));
 	if (!StoreKey.IsValid() || !RemoteNetTokenStoreState.ReserveTokenCount(NetToken.GetTypeId(), NetToken.GetIndex() + 1))
 	{
 		return false;
@@ -437,6 +449,8 @@ void FNetTokenStore::ReadTokenData(FNetSerializationContext& Context, const FNet
 {
 	if (NetToken.IsValid())
 	{
+		UE_IRIS_PARALLEL_EXPR(FScopedNetTokenStoreAccess ScopeAccess(this));
+
 		FNetBitStreamReader* Reader = Context.GetBitStreamReader();
 
 		FNetToken::FTypeId TokenTypeId = NetToken.GetTypeId();
@@ -466,6 +480,8 @@ void FNetTokenStore::ReadTokenData(FArchive& Ar, const FNetToken NetToken, FNetT
 {
 	if (NetToken.IsValid())
 	{
+		UE_IRIS_PARALLEL_EXPR(FScopedNetTokenStoreAccess ScopeAccess(this));
+
 		// Read type
 		FNetToken::FTypeId TokenTypeId = NetToken.GetTypeId();
 
@@ -547,6 +563,8 @@ void FNetTokenStore::AppendExport(FNetSerializationContext& Context, FNetToken N
 
 TArray<FNetToken> FNetTokenStore::GetAllNetTokens() const
 {
+	UE_IRIS_PARALLEL_EXPR(FScopedNetTokenStoreAccess ScopeAccess(this));
+	
 	TArray<FNetToken> Result;
 
 	const FNetTokenStoreState* TokenStoreState = GetLocalNetTokenStoreState();
@@ -563,8 +581,26 @@ TArray<FNetToken> FNetTokenStore::GetAllNetTokens() const
 	return MoveTemp(Result);
 }
 
+FScopedNetTokenStoreAccess::FScopedNetTokenStoreAccess(const FNetTokenStore* InTokenStore)
+{
+#if UE_SUPPORT_PARALLEL_IRIS
+	check(InTokenStore);
+	InTokenStore->NetTokenStoreCS.Lock();
+	TokenStore = InTokenStore;
+#endif // UE_SUPPORT_PARALLEL_IRIS
 }
 
+FScopedNetTokenStoreAccess::~FScopedNetTokenStoreAccess()
+{
+#if UE_SUPPORT_PARALLEL_IRIS
+	check(TokenStore);
+
+	TokenStore->NetTokenStoreCS.Unlock();
+#endif // UE_SUPPORT_PARALLEL_IRIS
+}
+
+} //namespace UE::Net
+
 uint32 UNetTokenTypeIdConfig::GetTypeID(const FString& TypeName) const
 {
 	// Need to reevaluate every retrieval in the event of hotfixes, or ini changes due to GFP loading/unloading.
diff --git a/Engine/Source/Runtime/Net/Iris/Private/Iris/ReplicationSystem/StringTokenStore.cpp b/Engine/Source/Runtime/Net/Iris/Private/Iris/ReplicationSystem/StringTokenStore.cpp
--- a/Engine/Source/Runtime/Net/Iris/Private/Iris/ReplicationSystem/StringTokenStore.cpp
+++ b/Engine/Source/Runtime/Net/Iris/Private/Iris/ReplicationSystem/StringTokenStore.cpp
@@ -49,6 +49,8 @@ FNetToken FStringTokenStore::GetOrCreateToken(const FString& String)
 
 FNetTokenDataStore::FNetTokenStoreKey FStringTokenStore::GetOrCreatePersistentString(const TCHAR* Name, uint32 Length)
 {
+	UE_IRIS_PARALLEL_EXPR(FScopedNetTokenStoreAccess ScopeAccess(&TokenStore));
+
 	// Hash name
 	const uint32 NameSize = Length * sizeof(TCHAR);
 	uint64 HashedName = CityHash64((const char*)Name, NameSize);
@@ -85,6 +87,7 @@ FStringTokenStore::FStringTokenStore(FNetTokenStore& InTokenStore)
 
 const TCHAR* FStringTokenStore::ResolveToken(FNetToken Token, const FNetTokenStoreState* NetTokenStoreState) const
 {
+	UE_IRIS_PARALLEL_EXPR(FScopedNetTokenStoreAccess ScopeAccess(&TokenStore));
 	const FNetTokenStoreState* TokenStoreState = TokenStore.IsLocalToken(Token) ? TokenStore.GetLocalNetTokenStoreState() : NetTokenStoreState;
 	if (Token.IsValid() && ensureMsgf(TokenStoreState, TEXT("FStringTokenStore::ResolveToken Needs valid TokenStoreState to resolve %s"), *Token.ToString()))
 	{
@@ -104,12 +107,14 @@ const TCHAR* FStringTokenStore::ResolveToken(FNetToken Token, const FNetTokenSto
 
 void FStringTokenStore::WriteTokenData(FNetSerializationContext& Context, FNetTokenStoreKey TokenStoreKey) const
 {
+	UE_IRIS_PARALLEL_EXPR(FScopedNetTokenStoreAccess ScopeAccess(&TokenStore));
 	// $TODO: $IRIS: Do not calculate the length of the string to write the data.
 	WriteString(Context.GetBitStreamWriter(), FStringView(StoredStrings[TokenStoreKey.GetKeyIndex()]));
 }
 
 void FStringTokenStore::WriteTokenData(FArchive& Ar, FNetTokenStoreKey TokenStoreKey, UPackageMap* Map /*= nullptr*/) const
 {
+	UE_IRIS_PARALLEL_EXPR(FScopedNetTokenStoreAccess ScopeAccess(&TokenStore));
 	FString Temp(StoredStrings[TokenStoreKey.GetKeyIndex()]);
 	Ar << Temp;
 }
diff --git a/Engine/Source/Runtime/Net/Iris/Public/Iris/IrisConfig.h b/Engine/Source/Runtime/Net/Iris/Public/Iris/IrisConfig.h
--- a/Engine/Source/Runtime/Net/Iris/Public/Iris/IrisConfig.h
+++ b/Engine/Source/Runtime/Net/Iris/Public/Iris/IrisConfig.h
@@ -43,4 +43,17 @@ IRISCORE_API EReplicationSystem GetUseIrisReplicationCmdlineValue();
 /** Enables code to simulate async loading stalls on specific objects on the client */
 #ifndef UE_NET_ASYNCLOADING_DEBUG
 #	define UE_NET_ASYNCLOADING_DEBUG !(UE_BUILD_SHIPPING || UE_BUILD_TEST)
+#endif
+
+// Controls whether we compile in thread safety mechanisms into Iris. Will add a small amount of overhead
+// due to mutex locking, but this is required when running Iris phases in parallel
+#ifndef UE_SUPPORT_PARALLEL_IRIS
+#define UE_SUPPORT_PARALLEL_IRIS 0
+#endif
+
+/** Enables thread safety logic inside UE_IRIS_PARALLEL_EXPR brackets - required when running Iris in parallel */
+#if UE_SUPPORT_PARALLEL_IRIS
+#define UE_IRIS_PARALLEL_EXPR(expr) expr
+#else
+#define UE_IRIS_PARALLEL_EXPR(expr) (void)0
 #endif
\ No newline at end of file
diff --git a/Engine/Source/Runtime/Net/Iris/Public/Iris/ReplicationSystem/NameTokenStore.h b/Engine/Source/Runtime/Net/Iris/Public/Iris/ReplicationSystem/NameTokenStore.h
--- a/Engine/Source/Runtime/Net/Iris/Public/Iris/ReplicationSystem/NameTokenStore.h
+++ b/Engine/Source/Runtime/Net/Iris/Public/Iris/ReplicationSystem/NameTokenStore.h
@@ -4,6 +4,8 @@
 
 #include "CoreTypes.h"
 #include "Containers/Map.h"
+#include "Misc/TransactionallySafeRWLock.h"
+#include "Iris/IrisConfig.h"
 #include "Iris/ReplicationSystem/NetTokenStore.h"
 
 namespace UE::Net
@@ -42,6 +44,11 @@ private:
 
 	inline static FName NameTokenStoreName = "NameTokenStore";
 
+#if UE_SUPPORT_PARALLEL_IRIS
+	//Used to protect parallel access to FNameToKey and StoredFNames. Can be called within a FScopedNetTokenStoreAccess scope, but code that already locks this mutex should not lock TokenStoreAccess afterwards (avoid lock order inversion)
+	mutable FTransactionallySafeRWLock NameTokenStoreMutex;
+#endif // UE_SUPPORT_PARALLEL_IRIS
+
 	TMap<FName, FNetTokenStoreKey> FNameToKey;
 	TArray<FName> StoredFNames;
 };
diff --git a/Engine/Source/Runtime/Net/Iris/Public/Iris/ReplicationSystem/NetBlob/NetBlob.h b/Engine/Source/Runtime/Net/Iris/Public/Iris/ReplicationSystem/NetBlob/NetBlob.h
--- a/Engine/Source/Runtime/Net/Iris/Public/Iris/ReplicationSystem/NetBlob/NetBlob.h
+++ b/Engine/Source/Runtime/Net/Iris/Public/Iris/ReplicationSystem/NetBlob/NetBlob.h
@@ -131,13 +131,31 @@ public:
 	IRISCORE_API void CollectObjectReferences(FNetSerializationContext& Context, FNetReferenceCollector& Collector) const;
 
 	/** Adds a reference. A blob is created with reference count zero. */
-	void AddRef() const { ++RefCount; }
+	void AddRef() const
+	{ 
+		UE_AUTORTFM_OPEN
+		{
+			++RefCount;
+		};
+		UE_AUTORTFM_ONABORT(this)
+		{
+			--RefCount;
+		};
+	}
 
 	/** Removes a reference. When the reference count reaches zero the blob will be deleted. A blob is created with reference count zero. */
 	IRISCORE_API void Release() const;
 
 	/** Returns the reference count.  A blob is created with reference count zero. */
-	int32 GetRefCount() const { return RefCount; }
+	int32 GetRefCount() const
+	{
+		int32 ReturnVal = 0;
+		UE_AUTORTFM_OPEN
+		{
+			ReturnVal = RefCount;
+		};
+		return ReturnVal;
+	}
 
 	/* Returns true if the blob has additional exports to add. */
 	bool HasExports() const { return EnumHasAnyFlags(GetCreationInfo().Flags, ENetBlobFlags::HasExports); }
@@ -178,7 +196,7 @@ protected:
 	FQuantizedBlobState QuantizedBlobState;
 
 private:
-	mutable int32 RefCount;
+	mutable std::atomic<int32> RefCount;
 };
 
 /**
diff --git a/Engine/Source/Runtime/Net/Iris/Public/Iris/ReplicationSystem/NetTokenStore.h b/Engine/Source/Runtime/Net/Iris/Public/Iris/ReplicationSystem/NetTokenStore.h
--- a/Engine/Source/Runtime/Net/Iris/Public/Iris/ReplicationSystem/NetTokenStore.h
+++ b/Engine/Source/Runtime/Net/Iris/Public/Iris/ReplicationSystem/NetTokenStore.h
@@ -3,6 +3,8 @@
 #pragma once
 
 #include "CoreTypes.h"
+#include "Iris/IrisConfig.h"
+#include "Misc/TransactionallySafeCriticalSection.h"
 #include "Net/Core/NetToken/NetToken.h"
 #include "UObject/NameTypes.h"
 #include "Containers/Array.h"
@@ -180,6 +182,22 @@ private:
 	FNetToken::FTypeId TypeId;
 };
 
+/**
+* A helper RAII class to allow thread-safe access to NetTokenStore
+* When making calls to the FNetTokenStore in a parallel phase, first create FScopedNetTokenStoreAccess to
+* ensure that the access is threadsafe during the lifetime of the FScopedNetTokenStoreAccess object.
+*/
+class FScopedNetTokenStoreAccess
+{
+public:
+	IRISCORE_API explicit FScopedNetTokenStoreAccess(const FNetTokenStore* InTokenStore);
+
+	IRISCORE_API ~FScopedNetTokenStoreAccess();
+
+private:
+	const FNetTokenStore* TokenStore = nullptr;
+};
+
 /**
  * FNetTokenStore
  * Main system for using NetTokensm currently owns type specific NetTokenDataStores and per connection state
@@ -369,7 +387,15 @@ public:
 	{
 		return InternalReadNetToken(Ar, FNetToken::InvalidTokenTypeId);
 	}
-	
+
+#if UE_SUPPORT_PARALLEL_IRIS
+protected:
+	friend FScopedNetTokenStoreAccess;
+
+	//Coarse grain mutex which is used to protect access of all NetTokenStore, NetTokenStoreState and NetTokenDataStore. 
+	mutable FTransactionallySafeCriticalSection NetTokenStoreCS;
+#endif
+
 private:
 
 	friend UNetTokenDataStream;
@@ -407,6 +433,7 @@ private:
 
 inline FNetTokenDataStore::FNetTokenStoreKey FNetTokenDataStore::GetNextNetTokenStoreKey()
 {
+	UE_IRIS_PARALLEL_EXPR(FScopedNetTokenStoreAccess ScopeAccess(&TokenStore));
 	uint32 NextTokenStoreKeyIndex = (uint32)StoredTokens.Num();
 	if (ensure(NextTokenStoreKeyIndex < FNetToken::MaxNetTokenCount))
 	{
diff --git a/Engine/Source/Runtime/Net/Iris/Public/Iris/ReplicationSystem/StructNetTokenDataStore.h b/Engine/Source/Runtime/Net/Iris/Public/Iris/ReplicationSystem/StructNetTokenDataStore.h
--- a/Engine/Source/Runtime/Net/Iris/Public/Iris/ReplicationSystem/StructNetTokenDataStore.h
+++ b/Engine/Source/Runtime/Net/Iris/Public/Iris/ReplicationSystem/StructNetTokenDataStore.h
@@ -308,6 +308,8 @@ protected:
 	// Creates a persistent copy of the input struct 
 	FNetTokenStoreKey GetOrCreatePersistentState(const T& InNetTokenData)
 	{
+		UE_IRIS_PARALLEL_EXPR(FScopedNetTokenStoreAccess ScopeAccess(&TokenStore));
+
 		LLM_SCOPE_BYTAG(NetTokenStructState);
 		const uint64 HashKey = InNetTokenData.GetUniqueKey();
 		if (HashToKey.Contains(HashKey))
