Fix for interpreting command-line options containing the cvar "net.Iris.UseIrisReplication" as the explicit command-line option "-UseIrisReplication=". Setting the "net.Iris.UseIrisReplication" cvar on the command-line will no longer override FIrisNetDriverConfig::PreferredReplicationSystem.

diff --git a/Engine/Source/Runtime/Net/Iris/Private/Iris/IrisConfig.cpp b/Engine/Source/Runtime/Net/Iris/Private/Iris/IrisConfig.cpp
--- a/Engine/Source/Runtime/Net/Iris/Private/Iris/IrisConfig.cpp
+++ b/Engine/Source/Runtime/Net/Iris/Private/Iris/IrisConfig.cpp
@@ -25,7 +25,7 @@ void SetUseIrisReplication(bool EnableIrisReplication)
 EReplicationSystem GetUseIrisReplicationCmdlineValue()
 {
 	int32 UseIrisReplication=0;
-	if (FParse::Value(FCommandLine::Get(), TEXT("UseIrisReplication="), UseIrisReplication))
+	if (FParse::Value(FCommandLine::Get(), TEXT("-UseIrisReplication="), UseIrisReplication))
 	{
 		// Try to force the requested system if the cmdline is present
 		return UseIrisReplication > 0 ? EReplicationSystem::Iris : EReplicationSystem::Generic;
