Use the close reason when destroying a replicated sub-object in Iris.

diff --git a/Engine/Source/Runtime/Net/Iris/Private/Iris/ReplicationSystem/ReplicationReader.cpp b/Engine/Source/Runtime/Net/Iris/Private/Iris/ReplicationSystem/ReplicationReader.cpp
--- a/Engine/Source/Runtime/Net/Iris/Private/Iris/ReplicationSystem/ReplicationReader.cpp
+++ b/Engine/Source/Runtime/Net/Iris/Private/Iris/ReplicationSystem/ReplicationReader.cpp
@@ -640,7 +640,7 @@ uint32 FReplicationReader::ReadRootObjectsPendingDestroy(FNetSerializationContex
 						// End replication for all subobjects and destroy dynamic ones
 						const FNetRefHandle SubObjectRefHandle = NetRefHandleManager->GetNetRefHandleFromInternalIndex(SubObjectIndex);
 						const bool bShouldDestroySubObject = SubObjectRefHandle.IsDynamic();
-						EndReplication(SubObjectIndex, bTearOff, bShouldDestroySubObject);
+						EndReplication(SubObjectIndex, bTearOff, bShouldDestroySubObject, CloseReason);
 					}
 				}
 			}
@@ -2646,7 +2646,7 @@ void FReplicationReader::DispatchEndReplication(FNetSerializationContext& Contex
 					// End replication for all subobjects and destroy dynamic ones
 					const FNetRefHandle SubObjectRefHandle = NetRefHandleManager->GetNetRefHandleFromInternalIndex(SubObjectIndex);
 					const bool bShouldDestroySubObject = SubObjectRefHandle.IsDynamic();
-					EndReplication(SubObjectIndex, Info.bTearOff, bShouldDestroySubObject);
+					EndReplication(SubObjectIndex, Info.bTearOff, bShouldDestroySubObject, Info.CloseReason);
 				}
 			}
 		}
@@ -2886,7 +2886,7 @@ void FReplicationReader::ProcessQueuedBatches()
 							// End replication for all subobjects and destroy dynamic ones
 							const FNetRefHandle SubObjectRefHandle = NetRefHandleManager->GetNetRefHandleFromInternalIndex(SubObjectIndex);
 							const bool bShouldDestroySubObject = SubObjectRefHandle.IsDynamic();
-							EndReplication(SubObjectIndex, bTearOff, bShouldDestroySubObject);
+							EndReplication(SubObjectIndex, bTearOff, bShouldDestroySubObject, CloseReason);
 						}
 					}
 					else
