Iris
* Added StartReplication param that forces the replication system to assign a new NetRefHandle if the object was previously replication but is now stopped.
* This is useful when core traits of the object changed since the last time it was replicated: New outer or no longer stable.

#rnx

diff --git a/Engine/Source/Runtime/Net/Iris/Private/Iris/ReplicationSystem/NetRefHandleManager.h b/Engine/Source/Runtime/Net/Iris/Private/Iris/ReplicationSystem/NetRefHandleManager.h
--- a/Engine/Source/Runtime/Net/Iris/Private/Iris/ReplicationSystem/NetRefHandleManager.h
+++ b/Engine/Source/Runtime/Net/Iris/Private/Iris/ReplicationSystem/NetRefHandleManager.h
@@ -362,7 +362,7 @@ public:
 
 	void AddReferencedObjects(FReferenceCollector& Collector);
 
-	bool GetIsDestroyedStartupObject(FInternalNetRefIndex InternalIndex) const { return DestroyedStartupObjectInternalIndices.GetBit(InternalIndex); }
+	bool IsDestroyedStartupObject(FInternalNetRefIndex InternalIndex) const { return DestroyedStartupObjectInternalIndices.GetBit(InternalIndex); }
 	inline uint32 GetOriginalDestroyedStartupObjectIndex(FInternalNetRefIndex InternalIndex) const;
 
 	const FNetBitArrayView GetDestroyedStartupObjectInternalIndices() const { return MakeNetBitArrayView(DestroyedStartupObjectInternalIndices); }
diff --git a/Engine/Source/Runtime/Net/Iris/Private/Iris/ReplicationSystem/ObjectReferenceCache.cpp b/Engine/Source/Runtime/Net/Iris/Private/Iris/ReplicationSystem/ObjectReferenceCache.cpp
--- a/Engine/Source/Runtime/Net/Iris/Private/Iris/ReplicationSystem/ObjectReferenceCache.cpp
+++ b/Engine/Source/Runtime/Net/Iris/Private/Iris/ReplicationSystem/ObjectReferenceCache.cpp
@@ -57,10 +57,6 @@ DEFINE_LOG_CATEGORY_STATIC(LogIrisReferences, Log, UE_NET_REFERENCECACHE_LOG_COM
 #define UE_LOG_REFERENCECACHE(Verbosity, Format, ...)  UE_LOG(LogIrisReferences, Verbosity, Format, ##__VA_ARGS__)
 #define UE_CLOG_REFERENCECACHE(Condition, Verbosity, Format, ...)  UE_CLOG(Condition, LogIrisReferences, Verbosity, Format, ##__VA_ARGS__)
 
-// $TODO: If we add information about why a dynamic object is being destroyed (filtering or actual destroy) then we can 
-// try to clean up references to reduce cache pollution
-#define UE_NET_CLEANUP_REFERENCES_WITH_DYNAMIC_OUTER 0
-
 namespace UE::Net::Private
 {
 
@@ -86,6 +82,16 @@ FAutoConsoleVariableRef CVarIrisAsyncLoadingFailAllLoads(TEXT("net.Iris.AsyncLoa
 static bool bIrisEnableRemoveReferenceFallbackPath = true;
 FAutoConsoleVariableRef CVarIrisEnableRemoveReferenceFallbackPath(TEXT("net.iris.ObjectReferenceCache.EnableRemoveReferenceFallbackPath"), bIrisEnableRemoveReferenceFallbackPath, TEXT("Whether to enable the slow fallback path or not when ObjectReferenceCache::RemoveReference is passed a null pointer object."), ECVF_Default);
 
+/**
+ * Don't allow infinite recursion of InternalLoadObject - an attacker could
+ * send malicious packets that cause a stack overflow on the server.
+ */
+static const int INTERNAL_READ_REF_RECURSION_LIMIT = 16;
+
+
+//------------------------------------------------------------------------
+// FPendingAsyncLoadRequest
+//------------------------------------------------------------------------
 
 FObjectReferenceCache::FPendingAsyncLoadRequest::FPendingAsyncLoadRequest(FNetRefHandle InNetRefHandle, double InRequestStartTime)
 : RequestStartTime(InRequestStartTime)
@@ -106,16 +112,11 @@ void FObjectReferenceCache::FPendingAsyncLoadRequest::Merge(FNetRefHandle InNetR
 	NetRefHandles.AddUnique(InNetRefHandle);
 }
 
-/**
- * Don't allow infinite recursion of InternalLoadObject - an attacker could
- * send malicious packets that cause a stack overflow on the server.
- */
-static const int INTERNAL_READ_REF_RECURSION_LIMIT = 16;
+//------------------------------------------------------------------------
+// FObjectReferenceCache
+//------------------------------------------------------------------------
 
 FObjectReferenceCache::FObjectReferenceCache()
-: bIsAuthority(false)
-, AsyncLoadMode(EAsyncLoadMode::UseCVar)
-, bCachedCVarAllowAsyncLoading(false)
 {
 	SetAsyncLoadMode(EAsyncLoadMode::UseCVar);
 }
@@ -144,7 +145,6 @@ bool FObjectReferenceCache::SupportsObjectInternal(const UObject* Object) const
 	return false;
 }
 
-// $IRIS: $TODO: refactor into utility methods as this is coped from GUIDCache
 bool FObjectReferenceCache::IsDynamicObject(const UObject* Object) const
 {
 	checkSlow(Object != nullptr);
@@ -160,11 +160,6 @@ bool FObjectReferenceCache::IsDynamicInternal(const UObject* Object) const
 	return !Object->IsFullNameStableForNetworking();
 }
 
-bool FObjectReferenceCache::IsAuthority() const
-{
-	return bIsAuthority;
-}
-
 bool FObjectReferenceCache::CanClientLoadObjectInternal(const UObject* Object, bool bIsDynamic) const
 {
 	// PackageMapClient can't load maps, we must wait for the client to load the map when ready
@@ -251,6 +246,29 @@ FNetRefHandle FObjectReferenceCache::CreateObjectReferenceHandle(const UObject*
 	return Ref.GetRefHandle();
 }
 
+FNetRefHandle FObjectReferenceCache::ReassignObjectRefHandle(const UObject* Object, UReplicationSystem* ReplicationSystem)
+{
+	// Clear the handle already assigned to the pointer.
+	const FSetElementId ObjectToNetRefId = ObjectToNetReferenceHandle.FindId(Object);
+	if (ObjectToNetRefId.IsValidId())
+	{
+		// TODO: Should we remove old NetRefHandle -> Object map entry ?
+		//		 It depends if the RepSystem will want to access the object via the old NetRefHandle later.
+		//		 But since the handle is supposed to be in a Destroyed/StopReplication state its already dangerous to access  the object.
+		//		 We might want to mark the entry + netrefhandle as deprecated on StopReplication directly and only access the object via an Explicit call if it is actually neeeded.
+		//const FNetRefHandle DeprecatedRefHandle = ObjectToNetReferenceHandle.Get(ObjectToNetRefId);
+		//ReferenceHandleToCachedReference.Remove(DeprecatedRefHandle);
+
+		// Only remove the object->handle map for now.
+		ObjectToNetReferenceHandle.Remove(ObjectToNetRefId);
+	}
+
+	FNetObjectReference ObjectRef;
+	bool bSuccess = CreateObjectReferenceInternal(Object, ObjectRef, ReplicationSystem);
+
+	return bSuccess ? ObjectRef.GetRefHandle() : FNetRefHandle::GetInvalid();
+}
+
 FNetRefHandle FObjectReferenceCache::PreRegisterObjectReferenceHandle(const UObject* Object, UReplicationSystem* ReplicationSystem)
 {
 	if (IsAuthority())
@@ -565,8 +583,9 @@ void FObjectReferenceCache::BindRemoteReference(FNetRefHandle RefHandle, const U
 		{
 			if (ExistingCachedObject->NetRefHandle != RefHandle)
 			{
-				UE_LOG_REFERENCECACHE(Warning, TEXT("%hs: Detected conflicting Ref %s: Received: %s, Object: %s (0x%p)"), __FUNCTION__, *ExistingCachedObject.NetRefHandle.ToString(),
+				UE_LOG_REFERENCECACHE(Warning, TEXT("%hs: Detected conflicting Ref %s: Received: %s, Object: %s (0x%p)"), __FUNCTION__, *ExistingCachedObject->NetRefHandle.ToString(),
 									  *RefHandle.ToString(), *GetNameSafe(Object), Object);
+				ensure(false);
 			}
 		}
 	}
@@ -795,7 +814,6 @@ bool FObjectReferenceCache::IsNetRefHandlePending(FNetRefHandle NetRefHandle, co
 	return false;
 }
 
-// $IRIS: $TODO: Most of the logic comes from GUIDCache::GetObjectFromNetGUID so we want to keep this in sync
 UObject* FObjectReferenceCache::ResolveObjectReferenceHandleInternal(FNetRefHandle RefHandle, const FNetObjectResolveContext& ResolveContext, bool& bOutMustBeMapped)
 {
 	FCachedNetObjectReference* CacheObjectPtr = ReferenceHandleToCachedReference.Find(RefHandle);
@@ -816,15 +834,15 @@ UObject* FObjectReferenceCache::ResolveObjectReferenceHandleInternal(FNetRefHand
 		// $IRIS: $TODO: Consider using FNames over string tokens, or store FNames in cache and fixup when resolving?
 #if UE_NET_VALIDATE_REFERENCECACHE
 		{			
-			FString ObjectPath(StringTokenStore->ResolveRemoteToken(CacheObjectPtr->RelativePath, *ResolveContext.RemoteNetTokenStoreState));
-			constexpr bool bReading = true;
-			RenamePathForPie(ResolveContext.ConnectionId, ObjectPath, bReading, ResolveContext.ReplicationSystem);
-			check(ObjectPath == Object->GetFName().ToString() || RefHandle.IsDynamic() || IsAuthority());
+			//FString ObjectPath(StringTokenStore->ResolveRemoteToken(CacheObjectPtr->RelativePath, *ResolveContext.RemoteNetTokenStoreState));
+			//constexpr bool bReading = true;
+			//RenamePathForPie(ResolveContext.ConnectionId, ObjectPath, bReading, ResolveContext.ReplicationSystem);
+			//check(ObjectPath == Object->GetFName().ToString() || RefHandle.IsDynamic() || IsAuthority());
 
 			FNetRefHandle* ExistingRefHandle = ObjectToNetReferenceHandle.Find(Object);
 			if (ExistingRefHandle && *ExistingRefHandle != RefHandle)
 			{
-				UE_LOG_REFERENCECACHE(Warning, TEXT("ObjectReferenceCache::ResolveObjectRefererenceHandle detected potential conflicting references: %s and %s"), *RefHandle.ToString(), *(*ExistingRefHandle).ToString());
+				UE_LOG_REFERENCECACHE(Warning, TEXT("ObjectReferenceCache::ResolveObjectRefererenceHandle detected potential conflicting references: %s and %s"), *RefHandle.ToString(), *(ExistingRefHandle->ToString()));
 			}
 		}
 #endif
diff --git a/Engine/Source/Runtime/Net/Iris/Private/Iris/ReplicationSystem/ObjectReferenceCache.h b/Engine/Source/Runtime/Net/Iris/Private/Iris/ReplicationSystem/ObjectReferenceCache.h
--- a/Engine/Source/Runtime/Net/Iris/Private/Iris/ReplicationSystem/ObjectReferenceCache.h
+++ b/Engine/Source/Runtime/Net/Iris/Private/Iris/ReplicationSystem/ObjectReferenceCache.h
@@ -38,6 +38,11 @@ struct FObjectReferenceCacheParams
 	FNetRefHandleSerialToObject NetRefHandleSerialToObject;
 };
 
+// $TODO: If we add information about why a dynamic object is being destroyed (filtering or actual destroy) then we can 
+// try to clean up references to reduce cache pollution
+#define UE_NET_CLEANUP_REFERENCES_WITH_DYNAMIC_OUTER 0
+
+
 // A lot of code in this class is extracted from the re-factored GUIDCache/ClientPackageMap and must be kept up to sync once they are submitted
 // Hopefully we can merge parts of this back together later on
 
@@ -53,11 +58,17 @@ public:
 	bool IsDynamicObject(const UObject* Object) const;
 
 	// Are we allowed to create new NetHandles to reference objects?
-	bool IsAuthority() const;
+	bool IsAuthority() const
+	{
+		return bIsAuthority;
+	}
 
 	// Create and assign a new NetHandle to the object
 	FNetRefHandle CreateObjectReferenceHandle(const UObject* Object, UReplicationSystem* ReplicationSystem);
 
+	// Remove any previous Object->NetRefHandle maps and ensures a new NetRefHandle is assigned to the object.
+	FNetRefHandle ReassignObjectRefHandle(const UObject* Object, UReplicationSystem* ReplicationSystem);
+
 	// Create and assign a new pre-registered NetHandle to the object
 	FNetRefHandle PreRegisterObjectReferenceHandle(const UObject* Object, UReplicationSystem* ReplicationSystem);
 
@@ -251,12 +262,11 @@ private:
 	// Map ReferenceHandle -> CachedReference
 	TMap<FNetRefHandle, FCachedNetObjectReference> ReferenceHandleToCachedReference;
 
+#if UE_NET_CLEANUP_REFERENCES_WITH_DYNAMIC_OUTER
 	// To properly clean up stale references referencing dynamic objects we need to track them
 	TMultiMap<FNetRefHandle, FNetRefHandle> HandleToDynamicOuter;
+#endif
 	
-	// Do we have authority to create references?
-	uint32 bIsAuthority : 1;
-
 	/**
 	 * Set of all current Objects that we've been requested to be referenced while we are doing async loading.
 	 * This is used to prevent objects (especially async load objects,
@@ -265,12 +275,16 @@ private:
 	 */
 	TMap<FNetRefHandle, FQueuedBatchObjectReference> QueuedBatchObjectReferences;
 
-	EAsyncLoadMode AsyncLoadMode;
-	bool bCachedCVarAllowAsyncLoading;
+	EAsyncLoadMode AsyncLoadMode = EAsyncLoadMode::UseCVar;
+	bool bCachedCVarAllowAsyncLoading = false;
+	
+	// Do we have authority to create references?
+	bool bIsAuthority = false;
 
 	/** Set of packages that are currently pending async loads, referenced by package name. */
 	TMap<FName, FPendingAsyncLoadRequest> PendingAsyncLoadRequests;
 
+	/** Optional delegates that can be used to override the UObject<->NetRefHandle id assignments */
 	FObjectToNetRefHandleSerial ObjectToNetRefHandleSerial;
 	FNetRefHandleSerialToObject NetRefHandleSerialToObject;
 
@@ -295,8 +309,8 @@ private:
 
 inline UObject* FObjectReferenceCache::ResolveObjectReferenceHandle(FNetRefHandle RefHandle, const FNetObjectResolveContext& ResolveContext)
 {
-	bool bMustBeMapped;
-	return ResolveObjectReferenceHandleInternal(RefHandle, ResolveContext, bMustBeMapped);
+	bool bOutMustBeMapped;
+	return ResolveObjectReferenceHandleInternal(RefHandle, ResolveContext, bOutMustBeMapped);
 }
 
 inline UObject* FObjectReferenceCache::ResolveObjectReference(const FNetObjectReference& ObjectRef, const FNetObjectResolveContext& ResolveContext)
diff --git a/Engine/Source/Runtime/Net/Iris/Private/Iris/ReplicationSystem/ObjectReplicationBridge.cpp b/Engine/Source/Runtime/Net/Iris/Private/Iris/ReplicationSystem/ObjectReplicationBridge.cpp
--- a/Engine/Source/Runtime/Net/Iris/Private/Iris/ReplicationSystem/ObjectReplicationBridge.cpp
+++ b/Engine/Source/Runtime/Net/Iris/Private/Iris/ReplicationSystem/ObjectReplicationBridge.cpp
@@ -439,7 +439,7 @@ void UObjectReplicationBridge::RemoveDestructionInfosForGroup(UE::Net::FNetObjec
 		TArray<FInternalNetRefIndex, TInlineAllocator<384>> ObjectIndicesToRemove;	
 		for (uint32 InternalObjectIndex : MakeArrayView(Group->Members))
 		{
-			if (NetRefHandleManager->GetIsDestroyedStartupObject(InternalObjectIndex))
+			if (NetRefHandleManager->IsDestroyedStartupObject(InternalObjectIndex))
 			{
 				const FNetRefHandleManager::FReplicatedObjectData& ObjectData = NetRefHandleManager->GetReplicatedObjectDataNoCheck(InternalObjectIndex);
 				if (StaticObjectsPendingDestroy.Remove(ObjectData.RefHandle))
@@ -567,27 +567,43 @@ void UObjectReplicationBridge::PreRegisterObjectWithReferenceHandle(const UObjec
 	ObjectReferenceCache->AddPreRegisteredReference(Handle, Instance, ReplicationSystem);
 }
 
-UE::Net::FNetRefHandle UObjectReplicationBridge::StartReplicatingNetObject(UObject* Instance, UE::Net::EReplicationFragmentTraits Traits, UE::Net::FNetObjectFactoryId NetFactoryId)
+UE::Net::FNetRefHandle UObjectReplicationBridge::StartReplicatingNetObject(UObject* Instance, UE::Net::EReplicationFragmentTraits Traits, UE::Net::FNetObjectFactoryId NetFactoryId, bool bForceNewHandle)
 {
 	using namespace UE::Net;
 	using namespace UE::Net::Private;
 
-	FNetRefHandle AllocatedRefHandle = ObjectReferenceCache->CreateObjectReferenceHandle(Instance, GetReplicationSystem());
+	FNetRefHandle AllocatedRefHandle = ObjectReferenceCache->CreateObjectReferenceHandle(Instance, ReplicationSystem);
 
-	// If we failed to assign a handle, or if the Handle already is replicating, just return the handle
+	// If we failed to assign a handle
 	if (!AllocatedRefHandle.IsValid())
 	{
 		return FNetRefHandle::GetInvalid();
 	}
 
+	// The instance is already replicating, just return the handle
 	if (IsReplicatedHandle(AllocatedRefHandle))
 	{
+		ensureMsgf(bForceNewHandle == false, TEXT("bForceNewHandle was requested for %s but it is already replicating. The object must stop replicating first."), *PrintObjectFromNetRefHandle(AllocatedRefHandle));
 		return AllocatedRefHandle;
 	}
 
+	//TODO: We could decide to automatically assign a new handle to objects that changed from Stable to Dynamic
+	//		Or objects that were renamed or changed their Outer since the last time they were replicating.
+	//		We could even do it for objects that can change protocols.
+	if (bForceNewHandle)
+	{
+		AllocatedRefHandle = ObjectReferenceCache->ReassignObjectRefHandle(Instance, ReplicationSystem);
+
+		if (!AllocatedRefHandle.IsValid())
+		{
+			ensureMsgf(AllocatedRefHandle.IsValid(), TEXT("Reassignement for %s failed to return a valid handle"), *GetNameSafe(Instance));
+			return FNetRefHandle::GetInvalid();
+		}
+	}
+
 	IRIS_PROFILER_SCOPE(StartReplicatingNetObject);
 	
-	FFragmentRegistrationContext FragmentRegistrationContext(GetReplicationStateDescriptorRegistry(), GetReplicationSystem(), Traits, Instance);
+	FFragmentRegistrationContext FragmentRegistrationContext(GetReplicationStateDescriptorRegistry(), ReplicationSystem, Traits, Instance);
 
 	// For everything derived from UObject we can call the virtual function RegisterReplicationFragments	
 	CallRegisterReplicationFragments(Instance, FragmentRegistrationContext, EFragmentRegistrationFlags::None);
@@ -701,7 +717,7 @@ UE::Net::FNetRefHandle UObjectReplicationBridge::StartReplicatingNetObject(UObje
 	NetRefHandleManager->GetReplicatedObjectDataNoCheck(InternalReplicationIndex).NetFactoryId = NetFactoryId;
 
 	// Initialize conditionals
-	FReplicationSystemInternal* ReplicationSystemInternal = GetReplicationSystem()->GetReplicationSystemInternal();
+	FReplicationSystemInternal* ReplicationSystemInternal = ReplicationSystem->GetReplicationSystemInternal();
 	ReplicationSystemInternal->GetConditionals().InitPropertyCustomConditions(InternalReplicationIndex);
 
 	// Keep track of handles with object references for garbage collection's sake.
@@ -741,7 +757,7 @@ UE::Net::FNetRefHandle UObjectReplicationBridge::StartReplicatingRootObject(UObj
 	Traits |= Params.bNeedsPreUpdate ? EReplicationFragmentTraits::NeedsPreSendUpdate : EReplicationFragmentTraits::None;
 
 	// Register the object
-	FNetRefHandle RefHandle = StartReplicatingNetObject(Instance, Traits, NetFactoryId);
+	FNetRefHandle RefHandle = StartReplicatingNetObject(Instance, Traits, NetFactoryId, Params.bForceNewHandle);
 
 	if (!RefHandle.IsValid())
 	{
@@ -866,12 +882,14 @@ UE::Net::FNetRefHandle UObjectReplicationBridge::StartReplicatingSubObject(UObje
 	}
 
 	// Pre-registering subobjects is not supported
-	FNetRefHandle PotentiallyPreRegisteredHandle = GetObjectReferenceCache()->GetObjectReferenceHandleFromObject(Instance);
-	if (PotentiallyPreRegisteredHandle.IsValid() && IsNetRefHandlePreRegistered(PotentiallyPreRegisteredHandle))
 	{
-		UE_LOG_OBJECTREPLICATIONBRIDGE(Error, TEXT("StartReplicatingSubObject attempted to replicate pre-registered object %s."), *GetFullNameSafe(Instance));
-		ensureMsgf(false, TEXT("Attempted to replicate a pre-registered subobject"));
-		return FNetRefHandle::GetInvalid();
+		const FNetRefHandle PotentiallyPreRegisteredHandle = GetObjectReferenceCache()->GetObjectReferenceHandleFromObject(Instance);
+		if (PotentiallyPreRegisteredHandle.IsValid() && IsNetRefHandlePreRegistered(PotentiallyPreRegisteredHandle))
+		{
+			UE_LOG_OBJECTREPLICATIONBRIDGE(Error, TEXT("StartReplicatingSubObject attempted to replicate pre-registered object %s."), *GetFullNameSafe(Instance));
+			ensureMsgf(false, TEXT("Attempted to replicate a pre-registered subobject"));
+			return FNetRefHandle::GetInvalid();
+		}
 	}
 
 	FNetRefHandle SubObjectRefHandle = GetReplicatedRefHandle(Instance);
@@ -889,7 +907,7 @@ UE::Net::FNetRefHandle UObjectReplicationBridge::StartReplicatingSubObject(UObje
 	}
 
 	const EReplicationFragmentTraits Traits = EReplicationFragmentTraits::CanReplicate;
-	SubObjectRefHandle = StartReplicatingNetObject(Instance, Traits, NetFactoryId);
+	SubObjectRefHandle = StartReplicatingNetObject(Instance, Traits, NetFactoryId, Params.bForceNewHandle);
 
 	if (!SubObjectRefHandle.IsValid())
 	{
diff --git a/Engine/Source/Runtime/Net/Iris/Private/Iris/ReplicationSystem/ReplicationWriter.cpp b/Engine/Source/Runtime/Net/Iris/Private/Iris/ReplicationSystem/ReplicationWriter.cpp
--- a/Engine/Source/Runtime/Net/Iris/Private/Iris/ReplicationSystem/ReplicationWriter.cpp
+++ b/Engine/Source/Runtime/Net/Iris/Private/Iris/ReplicationSystem/ReplicationWriter.cpp
@@ -516,7 +516,7 @@ void FReplicationWriter::StartReplication(uint32 InternalIndex)
 	Info.LastAckedBaselineIndex = FDeltaCompressionBaselineManager::InvalidBaselineIndex;
 	Info.PendingBaselineIndex = FDeltaCompressionBaselineManager::InvalidBaselineIndex;
 
-	const uint32 bIsDestructionInfo = NetRefHandleManager->GetIsDestroyedStartupObject(InternalIndex) ? 1U : 0U;	
+	const uint32 bIsDestructionInfo = NetRefHandleManager->IsDestroyedStartupObject(InternalIndex) ? 1U : 0U;	
 	if (bIsDestructionInfo)
 	{
 		// Check status of original object about to be destroyed, if it has been confirmed as created we do not replicate the destruction info object at all
@@ -2115,7 +2115,7 @@ uint32 FReplicationWriter::WriteRootObjectsPendingDestroy(FNetSerializationConte
 		WriteNetRefHandleId(Context, ObjectData.RefHandle);
 
 		// Write bit indicating if the static instance should be destroyed or not (could skip the bit for dynamic objects)
-		const bool bShouldDestroyInstance = ObjectData.RefHandle.IsDynamic() || NetRefHandleManager->GetIsDestroyedStartupObject(InternalIndex);
+		const bool bShouldDestroyInstance = ObjectData.RefHandle.IsDynamic() || NetRefHandleManager->IsDestroyedStartupObject(InternalIndex);
 		Writer.WriteBool(bShouldDestroyInstance);
 
 		const bool bHasDestroyReason = ObjectData.CloseReason != ECloseReplicationReason::None;
@@ -2794,7 +2794,7 @@ FReplicationWriter::EWriteObjectStatus FReplicationWriter::WriteObjectAndSubObje
 				if (Info.SubObjectPendingDestroy)
 				{
 					ReplicatedDestroyHeaderFlags |= ReplicatedDestroyHeaderFlags_EndReplication;
-					const bool bShouldDestroyInstance = ObjectData.RefHandle.IsDynamic() || NetRefHandleManager->GetIsDestroyedStartupObject(InternalIndex);					
+					const bool bShouldDestroyInstance = ObjectData.RefHandle.IsDynamic() || NetRefHandleManager->IsDestroyedStartupObject(InternalIndex);					
 					ReplicatedDestroyHeaderFlags |= bShouldDestroyInstance ? ReplicatedDestroyHeaderFlags_DestroyInstance : ReplicatedDestroyHeaderFlags_None;
 				}
 
diff --git a/Engine/Source/Runtime/Net/Iris/Public/Iris/ReplicationSystem/ObjectReplicationBridge.h b/Engine/Source/Runtime/Net/Iris/Public/Iris/ReplicationSystem/ObjectReplicationBridge.h
--- a/Engine/Source/Runtime/Net/Iris/Public/Iris/ReplicationSystem/ObjectReplicationBridge.h
+++ b/Engine/Source/Runtime/Net/Iris/Public/Iris/ReplicationSystem/ObjectReplicationBridge.h
@@ -268,6 +268,9 @@ public:
 		/** When enabled we ignore the class config for this object and instead use the one specified by ExplicitDynamicFilter */
 		bool bUseExplicitDynamicFilter = false;
 
+		/** Set to true to prevent the replication system from reusing a NetRefHandle previously assigned to the object. Useful if the object had stopped replicating and that old representation is considered to be destroyed. */
+		bool bForceNewHandle = false;
+
 		/** 
 		* The name of the dynamic filter to use for this object (instead of asking the class config).
 		* Can be none so that no dynamic filter is assigned to the object.
@@ -298,6 +301,9 @@ public:
 
 		/** Optional enum to select where in the list this subobject will be inserted */
 		UE::Net::ESubObjectInsertionOrder InsertionOrder = UE::Net::ESubObjectInsertionOrder::None;
+
+		/** Set to true to prevent the replication system from reusing a NetRefHandle previously assigned to the object. Useful if the object had stopped replicating and that old representation is considered to be destroyed. */
+		bool bForceNewHandle = false;
 	};
 
 public:
@@ -481,7 +487,7 @@ private:
 	float GetPollFrequencyOfRootObject(UE::Net::Private::FInternalNetRefIndex RootObjectIndex) const;
 
 	/** Does all that is necessary to replicate any given UObject. Returns the object's handle if successful  */
-	UE::Net::FNetRefHandle StartReplicatingNetObject(UObject* Instance, UE::Net::EReplicationFragmentTraits Traits, UE::Net::FNetObjectFactoryId NetFactoryId);
+	UE::Net::FNetRefHandle StartReplicatingNetObject(UObject* Instance, UE::Net::EReplicationFragmentTraits Traits, UE::Net::FNetObjectFactoryId NetFactoryId, bool bForceNewHandle);
 
 	/** Retrieves the prioritizer to set for the given class. If bRequireForceEnabled the config needs to have bForceEnableOnAllInstances set in order for this method to return the configured prioritizer. Returns an invalid handle if no prioritizer should be set. */
 	UE::Net::FNetObjectPrioritizerHandle GetPrioritizer(const UClass* Class, bool bRequireForceEnabled);
