Iris
- Added API for marking an object as a DebugObject
- Added API for allowing DataStreams to write DebugData that does not count against bandwidth limits

diff --git a/Engine/Source/Runtime/Engine/Private/Net/Experimental/Iris/DataStreamChannel.cpp b/Engine/Source/Runtime/Engine/Private/Net/Experimental/Iris/DataStreamChannel.cpp
--- a/Engine/Source/Runtime/Engine/Private/Net/Experimental/Iris/DataStreamChannel.cpp
+++ b/Engine/Source/Runtime/Engine/Private/Net/Experimental/Iris/DataStreamChannel.cpp
@@ -28,6 +28,7 @@
 
 #include "PacketHandler.h"
 #include "ProfilingDebugging/CsvProfiler.h"
+#include "Logging/LogMacros.h"
 // IWYU pragma: end_keep
 
 namespace UE::Net::Private
@@ -262,11 +263,24 @@ void UDataStreamChannel::Tick()
 	IRIS_PROFILER_SCOPE(UDataStreamChannel_Tick);
 	LLM_SCOPE_BYTAG(Iris);
 
-	WriteData(EDataStreamWriteMode::Full);
+	const bool bWroteData = WriteData(EDataStreamWriteMode::Full);
 
+#if UE_NET_IRIS_ALLOW_DEBUG_DATA
+	// If we have debugdata, we write it as separate packets.
+	if (DataStreamManager->HasDebugDataToSend() && !IsPacketWindowFull())
+	{
+		// Need to flush
+		if (bWroteData)
+		{
+			IRIS_PROFILER_SCOPE(UDataStreamChannel_FlushNet);
+			Connection->FlushNet();	
+		}
+		WriteData(EDataStreamWriteMode::DebugData);
+	}
+#endif
 }
 
-void UDataStreamChannel::WriteData(UE::Net::EDataStreamWriteMode WriteMode)
+bool UDataStreamChannel::WriteData(UE::Net::EDataStreamWriteMode WriteMode)
 {
 	using namespace UE::Net;
 
@@ -281,6 +295,8 @@ void UDataStreamChannel::WriteData(UE::Net::EDataStreamWriteMode WriteMode)
 		.WriteMode = WriteMode,
 	};
 
+	const SIZE_T RecordCount = WriteRecords.Count();
+
 	if (WriteMode == EDataStreamWriteMode::Full)
 	{
 		int32 CurrentQueuedBits = Connection->QueuedBits + Connection->SendBuffer.GetNumBits();
@@ -354,7 +370,7 @@ void UDataStreamChannel::WriteData(UE::Net::EDataStreamWriteMode WriteMode)
 				Connection->FlushNet(true);
 				Connection->Close(ENetCloseResult::HostClosedConnection);
 			}
-			
+
 			return UDataStream::EWriteResult::NoData;
 		}
 
@@ -375,9 +391,17 @@ void UDataStreamChannel::WriteData(UE::Net::EDataStreamWriteMode WriteMode)
 		OutBunch.bReliable = !OpenAcked;
 		OutBunch.SerializeBits(BitStreamBuffer, BitWriter.GetPosBits());
 
+		const int32 PreSendBunchQueuedBits = Connection->QueuedBits + Connection->SendBuffer.GetNumBits();
+
 		constexpr bool bAllowMerging = false;
 		const FPacketIdRange PacketIds = SendBunch(&OutBunch, bAllowMerging);
 
+		// Ugly trick to allow DebugData to not count towards bandwidth limit
+		if (BeginWriteParams.WriteMode == EDataStreamWriteMode::DebugData)
+		{
+			Connection->QueuedBits = PreSendBunchQueuedBits - Connection->SendBuffer.GetNumBits();
+		}
+
 #if UE_NET_TRACE_ENABLED
 		// Since we steal the connection collector, we need to clear it out before the bunch goes out of scope
 		SetTraceCollector(OutBunch, nullptr);
@@ -405,7 +429,7 @@ void UDataStreamChannel::WriteData(UE::Net::EDataStreamWriteMode WriteMode)
 	// Begin the write, if we have nothing todo, just return
 	if (DataStreamManager->BeginWrite(BeginWriteParams) == UDataStream::EWriteResult::NoData)
 	{
-		return;
+		return false;
 	}
 
 	bool bContinueWriting = false;
@@ -436,6 +460,7 @@ void UDataStreamChannel::WriteData(UE::Net::EDataStreamWriteMode WriteMode)
 		Connection->FlushNet();	
 	}
 
+	return WriteRecords.Count() != RecordCount;
 }
 
 bool UDataStreamChannel::CanStopTicking() const
diff --git a/Engine/Source/Runtime/Engine/Private/Net/Experimental/Iris/DataStreamChannel.h b/Engine/Source/Runtime/Engine/Private/Net/Experimental/Iris/DataStreamChannel.h
--- a/Engine/Source/Runtime/Engine/Private/Net/Experimental/Iris/DataStreamChannel.h
+++ b/Engine/Source/Runtime/Engine/Private/Net/Experimental/Iris/DataStreamChannel.h
@@ -82,7 +82,7 @@ private:
 
 	void SendOpenBunch();
 
-	void WriteData(UE::Net::EDataStreamWriteMode WriteMode);
+	bool WriteData(UE::Net::EDataStreamWriteMode WriteMode);
 
 	TObjectPtr<UDataStreamManager> DataStreamManager = nullptr;
 
diff --git a/Engine/Source/Runtime/Net/Iris/Private/Iris/DataStream/DataStream.cpp b/Engine/Source/Runtime/Net/Iris/Private/Iris/DataStream/DataStream.cpp
--- a/Engine/Source/Runtime/Net/Iris/Private/Iris/DataStream/DataStream.cpp
+++ b/Engine/Source/Runtime/Net/Iris/Private/Iris/DataStream/DataStream.cpp
@@ -31,6 +31,11 @@ void UDataStream::Update(const FUpdateParameters& Params)
 {
 }
 
+bool UDataStream::HasDebugDataToSend() const
+{
+	return false;
+}
+
 const TCHAR* LexToString(const UDataStream::EDataStreamState State)
 {
 	static const TCHAR* Names[] = {
diff --git a/Engine/Source/Runtime/Net/Iris/Private/Iris/DataStream/DataStreamManager.cpp b/Engine/Source/Runtime/Net/Iris/Private/Iris/DataStream/DataStreamManager.cpp
--- a/Engine/Source/Runtime/Net/Iris/Private/Iris/DataStream/DataStreamManager.cpp
+++ b/Engine/Source/Runtime/Net/Iris/Private/Iris/DataStream/DataStreamManager.cpp
@@ -45,6 +45,7 @@ public:
 	void ReadData(UE::Net::FNetSerializationContext& context);
 	void ProcessPacketDeliveryStatus(UE::Net::EPacketDeliveryStatus Status, FDataStreamRecord const* Record);
 	bool HasAcknowledgedAllReliableData() const;
+	bool HasDebugDataToSend() const;
 	
 	ECreateDataStreamResult CreateStream(const FName StreamName, ECreateDataStreamFlags Flags = ECreateDataStreamFlags::None);
 	void CloseStream(const FName StreamName);
@@ -169,6 +170,11 @@ bool UDataStreamManager::HasAcknowledgedAllReliableData() const
 	return Impl->HasAcknowledgedAllReliableData();
 }
 
+bool UDataStreamManager::HasDebugDataToSend() const
+{
+	return Impl->HasDebugDataToSend();
+}
+
 bool UDataStreamManager::IsKnownStreamDefinition(const FName StreamName)
 {
 	const UDataStreamDefinitions* StreamDefinitions = GetDefault<UDataStreamDefinitions>();
@@ -943,6 +949,19 @@ void UDataStreamManager::FImpl::ProcessPacketDeliveryStatus(UE::Net::EPacketDeli
 	Records.Pop();
 }
 
+bool UDataStreamManager::FImpl::HasDebugDataToSend() const
+{
+	for (TObjectPtr<const UDataStream> Stream : Streams)
+	{
+		if (Stream && Stream->HasDebugDataToSend())
+		{
+			return true;
+		}
+	}
+
+	return false;
+}
+
 bool UDataStreamManager::FImpl::HasAcknowledgedAllReliableData() const
 {
 	for (TObjectPtr<const UDataStream> Stream : Streams)
diff --git a/Engine/Source/Runtime/Net/Iris/Private/Iris/ReplicationSystem/NetRefHandleManager.cpp b/Engine/Source/Runtime/Net/Iris/Private/Iris/ReplicationSystem/NetRefHandleManager.cpp
--- a/Engine/Source/Runtime/Net/Iris/Private/Iris/ReplicationSystem/NetRefHandleManager.cpp
+++ b/Engine/Source/Runtime/Net/Iris/Private/Iris/ReplicationSystem/NetRefHandleManager.cpp
@@ -268,6 +268,7 @@ FInternalNetRefIndex FNetRefHandleManager::InternalCreateNetObject(const FNetRef
 	ReplicatedObjectStateBuffers[InternalIndex] = nullptr;
 	Data.ReceiveStateBuffer = nullptr;
 	Data.bShouldPropagateChangedStates = 1U;
+	Data.bIsDebugObject = 0U;
 
 	++ActiveObjectCount;
 
@@ -871,9 +872,23 @@ bool FNetRefHandleManager::AddDependentObject(FNetRefHandle ParentRefHandle, FNe
 	FReplicatedObjectData& ParentObjectData = GetReplicatedObjectDataNoCheck(ParentInternalIndex);
 
 	// SubObjects cannot have dependent objects or be a dependent object (for now)
-	check(!DependentObjectData.IsSubObject() && !ParentObjectData.IsSubObject());
-	check(!SubObjectInternalIndices.GetBit(DependentObjectInternalIndex));
-	check(!SubObjectInternalIndices.GetBit(ParentInternalIndex));
+	if (DependentObjectData.IsSubObject() || ParentObjectData.IsSubObject() || SubObjectInternalIndices.GetBit(DependentObjectInternalIndex) || SubObjectInternalIndices.GetBit(ParentInternalIndex))
+	{
+		UE_LOG(LogIris, Warning, TEXT("FNetRefHandleManager::AddDependentObject: SubObjects cannot have or be dependent objects. Parent: %s child: %s."), *PrintObjectFromNetRefHandle(ParentRefHandle), *PrintObjectFromNetRefHandle(DependentObjectRefHandle));
+		ensure(!DependentObjectData.IsSubObject() && !ParentObjectData.IsSubObject());
+		ensure(!SubObjectInternalIndices.GetBit(DependentObjectInternalIndex));
+		ensure(!SubObjectInternalIndices.GetBit(ParentInternalIndex));
+		return false;
+	}
+
+	// DebugObjects cannot have dependent objects or be a dependent object (for now)
+	if (DependentObjectData.IsDebugObject() || ParentObjectData.IsDebugObject())
+	{
+		UE_LOG(LogIris, Warning, TEXT("FNetRefHandleManager::AddDependentObject: DebugObjects cannot have or be dependent objects. Parent: %s child: %s."), *PrintObjectFromNetRefHandle(ParentRefHandle), *PrintObjectFromNetRefHandle(DependentObjectRefHandle));
+		ensure(!DependentObjectData.IsDebugObject());
+		ensure(!ParentObjectData.IsDebugObject());
+		return false;
+	}
 
 	FNetDependencyData::FDependentObjectInfoArray& ParentDependentObjectsArray = SubObjects.GetOrCreateDependentObjectInfoArray(ParentInternalIndex);
 	FNetDependencyData::FInternalNetRefIndexArray& DependentParentObjectArray = SubObjects.GetOrCreateInternalIndexArray<FNetDependencyData::DependentParentObjects>(DependentObjectInternalIndex);
@@ -1141,6 +1156,15 @@ bool FNetRefHandleManager::AddingDependencyWouldCreateCircularDependency(FIntern
 	return false;
 }
 
+void FNetRefHandleManager::SetIsDebugObject(FNetRefHandle Handle)
+{
+	FInternalNetRefIndex ObjectInternalIndex = GetInternalIndex(Handle);
+	if (ObjectInternalIndex != InvalidInternalIndex && ensureMsgf(ReplicatedObjectRefCount[ObjectInternalIndex] == 0U, TEXT("SetIsDebugObject can only be called before we have started to replicate an object. %s"), *PrintObjectFromIndex(ObjectInternalIndex)))
+	{
+		ReplicatedObjectData[ObjectInternalIndex].bIsDebugObject = 1U;
+	}
+}
+
 void FNetRefHandleManager::SetShouldPropagateChangedStates(FInternalNetRefIndex ObjectInternalIndex, bool bShouldPropagateChangedStates)
 {
 	if (ObjectInternalIndex != InvalidInternalIndex)
diff --git a/Engine/Source/Runtime/Net/Iris/Private/Iris/ReplicationSystem/NetRefHandleManager.h b/Engine/Source/Runtime/Net/Iris/Private/Iris/ReplicationSystem/NetRefHandleManager.h
--- a/Engine/Source/Runtime/Net/Iris/Private/Iris/ReplicationSystem/NetRefHandleManager.h
+++ b/Engine/Source/Runtime/Net/Iris/Private/Iris/ReplicationSystem/NetRefHandleManager.h
@@ -127,6 +127,7 @@ public:
 				uint16 bWantsFullPoll : 1U;
 				uint16 bPendingEndReplication : 1U;
 				uint16 bHasCachedCreationInfo : 1U;
+				uint16 bIsDebugObject : 1U;
 			};
 		};
 
@@ -140,6 +141,9 @@ public:
 
 		/** Returns true if this is a child dependent object  */
 		bool IsDependentObject() const { return bIsDependentObject; }
+
+		/** Returns true if this is a considered a debugobject  */
+		bool IsDebugObject() const { return bIsDebugObject; }
 	};
 
 	// We want to handle both actual Networked and "network addressable" handles in the same way
@@ -322,6 +326,8 @@ public:
 	void SetShouldPropagateChangedStates(FNetRefHandle Handle, bool bShouldPropagateChangedStates);
 	void SetShouldPropagateChangedStates(FInternalNetRefIndex ObjectInternalIndex, bool bShouldPropagateChangedStates);
 
+	void SetIsDebugObject(FNetRefHandle Handle);
+
 	/** Return the maximum internal index that NetBitArrays and TArrays are currently allocated for. */
 	FInternalNetRefIndex GetCurrentMaxInternalNetRefIndex() const { return CurrentMaxInternalNetRefIndex; };
 
diff --git a/Engine/Source/Runtime/Net/Iris/Private/Iris/ReplicationSystem/ReplicationSystem.cpp b/Engine/Source/Runtime/Net/Iris/Private/Iris/ReplicationSystem/ReplicationSystem.cpp
--- a/Engine/Source/Runtime/Net/Iris/Private/Iris/ReplicationSystem/ReplicationSystem.cpp
+++ b/Engine/Source/Runtime/Net/Iris/Private/Iris/ReplicationSystem/ReplicationSystem.cpp
@@ -2025,6 +2025,16 @@ void UReplicationSystem::SetDeltaCompressionStatus(FNetRefHandle Handle, UE::Net
 	return DC.SetDeltaCompressionStatus(ObjectInternalIndex, Status);
 }
 
+void UReplicationSystem::SetIsDebugObject(FNetRefHandle Handle)
+{
+#if UE_NET_IRIS_ALLOW_DEBUG_DATA
+	UE::Net::Private::FNetRefHandleManager& NetRefHandleManager = Impl->ReplicationSystemInternal.GetNetRefHandleManager();	
+	NetRefHandleManager.SetIsDebugObject(Handle);
+#else
+	ensureMsgf(false, TEXT("SetIsDebugObject is only supported in non-shipping builds."));
+#endif
+}
+
 void UReplicationSystem::SetIsNetTemporary(FNetRefHandle Handle)
 {
 	UE::Net::Private::FNetRefHandleManager& NetRefHandleManager = Impl->ReplicationSystemInternal.GetNetRefHandleManager();
diff --git a/Engine/Source/Runtime/Net/Iris/Public/Iris/DataStream/DataStream.h b/Engine/Source/Runtime/Net/Iris/Public/Iris/DataStream/DataStream.h
--- a/Engine/Source/Runtime/Net/Iris/Public/Iris/DataStream/DataStream.h
+++ b/Engine/Source/Runtime/Net/Iris/Public/Iris/DataStream/DataStream.h
@@ -3,6 +3,7 @@
 #pragma once
 
 #include "CoreTypes.h"
+#include "Misc/EnumClassFlags.h"
 #include "DataStream.generated.h"
 
 class UDataStreamManager;
@@ -18,15 +19,17 @@ namespace UE::Net
 	}
 
 
-enum class EDataStreamWriteMode : unsigned
-{
-	// Allowed to write all data, this is the default WriteMode
-	Full,
+	enum class EDataStreamWriteMode : unsigned
+	{
+		// Allowed to write all data, this is the default WriteMode
+		Full,
 
-	// Only write data that should be sent after PostTickDispatch
-	PostTickDispatch,
-};
+		// Only write data that should be sent after PostTickDispatch
+		PostTickDispatch,
 
+		// Only write debug data that should not count against bandwidth limits
+		DebugData,
+	};
 }
 
 /**
@@ -191,6 +194,11 @@ public:
 	 */
 	virtual bool HasAcknowledgedAllReliableData() const PURE_VIRTUAL(HasAcknowledgedAllReliableData, return true;);
 
+	/** 
+	 * @return true if the stream indicates that is has debug data to write.
+	 */
+	IRISCORE_API virtual bool HasDebugDataToSend() const;
+
 	/**
 	 * Get name of DataStream
 	 */
diff --git a/Engine/Source/Runtime/Net/Iris/Public/Iris/DataStream/DataStreamManager.h b/Engine/Source/Runtime/Net/Iris/Public/Iris/DataStream/DataStreamManager.h
--- a/Engine/Source/Runtime/Net/Iris/Public/Iris/DataStream/DataStreamManager.h
+++ b/Engine/Source/Runtime/Net/Iris/Public/Iris/DataStream/DataStreamManager.h
@@ -71,6 +71,9 @@ public:
 	/** Returns true if calling HasAcknowledgedAllReliableData on all data streams also returns true. */
 	IRISCORE_API virtual bool HasAcknowledgedAllReliableData() const override;
 
+	/** Returns true if HasDebugDataToSend returns true for any data stream. */
+	IRISCORE_API virtual bool HasDebugDataToSend() const override;
+
 	// DataStreamManager specifics
 
 	/**
diff --git a/Engine/Source/Runtime/Net/Iris/Public/Iris/IrisConfig.h b/Engine/Source/Runtime/Net/Iris/Public/Iris/IrisConfig.h
--- a/Engine/Source/Runtime/Net/Iris/Public/Iris/IrisConfig.h
+++ b/Engine/Source/Runtime/Net/Iris/Public/Iris/IrisConfig.h
@@ -56,4 +56,9 @@ IRISCORE_API EReplicationSystem GetUseIrisReplicationCmdlineValue();
 #define UE_IRIS_PARALLEL_EXPR(expr) expr
 #else
 #define UE_IRIS_PARALLEL_EXPR(expr) (void)0
-#endif
\ No newline at end of file
+#endif
+
+/** Debug. */
+#ifndef UE_NET_IRIS_ALLOW_DEBUG_DATA
+#	define UE_NET_IRIS_ALLOW_DEBUG_DATA !(UE_BUILD_SHIPPING)
+#endif
diff --git a/Engine/Source/Runtime/Net/Iris/Public/Iris/ReplicationSystem/ReplicationSystem.h b/Engine/Source/Runtime/Net/Iris/Public/Iris/ReplicationSystem/ReplicationSystem.h
--- a/Engine/Source/Runtime/Net/Iris/Public/Iris/ReplicationSystem/ReplicationSystem.h
+++ b/Engine/Source/Runtime/Net/Iris/Public/Iris/ReplicationSystem/ReplicationSystem.h
@@ -716,6 +716,12 @@ public:
 	 */
 	IRISCORE_API void TearOffNextUpdate(FNetRefHandle Handle);
 
+	/**
+	  * Mark an object as a being a debugobject. 
+	  * NOTE: Only allowed in non-shipping builds
+	  */
+	IRISCORE_API void SetIsDebugObject(FNetRefHandle Handle);
+
 	/**
 	 * Force the passed object to be considered for replication this frame.
 	 * This will also force it's sub objects, root object and any of it's dependents to also be considered for replication.
