Iris
- Implemented debug write pass for ReplicationWriter that writes NetObjects marked as debugobjects using separate packets not counting againt bandwidth limit.

diff --git a/Engine/Source/Runtime/GameplayDebugger/Private/GameplayDebuggerCategoryReplicator.cpp b/Engine/Source/Runtime/GameplayDebugger/Private/GameplayDebuggerCategoryReplicator.cpp
--- a/Engine/Source/Runtime/GameplayDebugger/Private/GameplayDebuggerCategoryReplicator.cpp
+++ b/Engine/Source/Runtime/GameplayDebugger/Private/GameplayDebuggerCategoryReplicator.cpp
@@ -11,10 +11,14 @@
 #include "GameplayDebuggerExtension.h"
 #include "Net/UnrealNetwork.h"
 #include "Engine/NetDriver.h"
+#include "Engine/NetConnection.h"
 #include "VisualLogger/VisualLogger.h"
+#include "ProfilingDebugging/CpuProfilerTrace.h"
 
 #include "Iris/IrisConfig.h"
+#include "Iris/ReplicationSystem/ReplicationSystem.h"
 #include "Net/Iris/ReplicationSystem/ReplicationSystemUtil.h"
+#include "Net/Iris/ReplicationSystem/EngineReplicationBridge.h"
 
 #include UE_INLINE_GENERATED_CPP_BY_NAME(GameplayDebuggerCategoryReplicator)
 
@@ -126,6 +130,15 @@ void FGameplayDebuggerNetPack::PopulateFromOwner()
 	
 	if (Owner && Owner->bIsEnabled && Owner->Categories.Num() == SavedData.Num())
 	{
+		// If receiving end are stopped in debugger etc. we should not post more RPCs.
+		UNetConnection* Connection = Owner->GetNetConnection();
+		if (!Connection || !Connection->IsNetReady() || Connection->IsPacketSequenceWindowFull(4))
+		{
+			return;
+		}
+
+		TRACE_CPUPROFILER_EVENT_SCOPE(FGameplayDebugger_PopulateFromOwner);
+
 		// Update SavedCategories for replication, DataPacks-packets are handled using RPC`s
 		for (int32 Idx = 0; Idx < SavedData.Num(); Idx++)
 		{
@@ -147,18 +160,22 @@ void FGameplayDebuggerNetPack::PopulateFromOwner()
 				SavedCategory.Shapes = CategoryOb.ReplicatedShapes;
 			}
 
-			// Throttled send of DataPackRPC:s
+			// Send the datapack, rpcs. One per category.
 			if (Owner->bSendDataPacksUsingRPC)
 			{
+				// We use saveddata to track what we have sent
+				SavedCategory.DataPacks.SetNum(CategoryOb.ReplicatedDataPacks.Num());
+
 				for (int32 DataPackIdx = 0; DataPackIdx < CategoryOb.ReplicatedDataPacks.Num(); DataPackIdx++)
 				{
 					FGameplayDebuggerDataPack& DataPack = CategoryOb.ReplicatedDataPacks[DataPackIdx];
 					if (!CategoryOb.bIsLocal)
 					{
-						if ((DataPack.bNeedsConfirmation && !DataPack.bReceived))
+						if ((SavedCategory.DataPacks[DataPackIdx] != DataPack.Header) || (DataPack.bNeedsConfirmation && !DataPack.bReceived))
 						{
 							// Send the update data pack as an reliable rpc instead
 							Owner->SendDataPackPacket(CategoryOb.GetCategoryName(), DataPackIdx, DataPack);						
+							SavedCategory.DataPacks[DataPackIdx] = DataPack.Header;
 						}
 					}
 				}
@@ -174,6 +191,8 @@ void FGameplayDebuggerNetPack::ApplyToOwner()
 		return;
 	}
 
+	TRACE_CPUPROFILER_EVENT_SCOPE(FGameplayDebugger_ApplyToOwner);
+
 	bool bHasCategoryStateChanges = false;
 	for (const FGameplayDebuggerCategoryData& SavedDataEntry : SavedData)
 	{
@@ -551,6 +570,7 @@ AGameplayDebuggerCategoryReplicator::AGameplayDebuggerCategoryReplicator(const F
 	ReplicatedData.Owner = this;
 
 	bSendDataPacksUsingRPC = false;
+	bOnlyRelevantToOwner = true;
 }
 
 void AGameplayDebuggerCategoryReplicator::BeginPlay()
@@ -568,12 +588,18 @@ void AGameplayDebuggerCategoryReplicator::BeginPlay()
 
 void AGameplayDebuggerCategoryReplicator::OnReplicationStartedForIris(const FOnReplicationStartedParams& Params)
 {
-	if (UE::Net::FReplicationSystemUtil::GetReplicationSystem(this))
+	if (UReplicationSystem* ReplicationSystem = UE::Net::FReplicationSystemUtil::GetReplicationSystem(this))
 	{
 		// If iris is enabled, datapack-packets are always sent using RPC`s and we also requires PreReplication to called in order to populate replicated data outside of serialization.
 		SetCallPreReplication(true);
 		bSendDataPacksUsingRPC = true;
-		bOnlyRelevantToOwner = true;
+
+		// Mark this as a debugobject to allow it to write more data not counting towards the bandwidth limits
+		if (UEngineReplicationBridge* Bridge = ReplicationSystem->GetReplicationBridgeAs<UEngineReplicationBridge>())
+		{
+			ReplicationSystem->SetIsDebugObject(Bridge->GetReplicatedRefHandle(this));
+			ReplicationSystem->SetStaticPriority(Bridge->GetReplicatedRefHandle(this), 10.f);
+		}
 	}
 
 	Super::OnReplicationStartedForIris(Params);
@@ -972,12 +998,6 @@ void AGameplayDebuggerCategoryReplicator::CollectCategoryData(bool bForce)
 					{
 						UE_LOG(LogGameplayDebugReplication, Verbose, TEXT("Category[%d].DataPack[%d] SENT, DataVersion:%d DataSize:%d SyncCounter:%d"),
 							Idx, DataPackIdx, DataPack.Header.DataVersion, DataPack.Header.DataSize, DataPack.Header.SyncCounter);
-
-						// Send the update data pack as an reliable rpc instead, note this will just send the first part, multipart datapacks will be throttled over multiple frames
-						if (bSendDataPacksUsingRPC)
-						{
-							SendDataPackPacket(CategoryOb.GetCategoryName(), DataPackIdx, DataPack);						
-						}
 					}
 				}
 			}
diff --git a/Engine/Source/Runtime/Net/Iris/Private/Iris/ReplicationSystem/AttachmentReplication.cpp b/Engine/Source/Runtime/Net/Iris/Private/Iris/ReplicationSystem/AttachmentReplication.cpp
--- a/Engine/Source/Runtime/Net/Iris/Private/Iris/ReplicationSystem/AttachmentReplication.cpp
+++ b/Engine/Source/Runtime/Net/Iris/Private/Iris/ReplicationSystem/AttachmentReplication.cpp
@@ -272,6 +272,16 @@ void FNetObjectAttachmentSendQueue::DropUnreliable(bool& bOutHasUnsent)
 	}
 }
 
+SIZE_T FNetObjectAttachmentSendQueue::GetReliableCount() const
+{
+	if (ReliableQueue == nullptr)
+	{
+		return 0;
+	}
+
+	return ReliableQueue->GetUnsentBlobCount();
+}
+
 bool FNetObjectAttachmentSendQueue::HasUnsent() const
 {
 	return !UnreliableQueue.IsEmpty() || (ReliableQueue != nullptr && ReliableQueue->HasUnsentBlobs());
@@ -608,6 +618,17 @@ void FNetObjectAttachmentsWriter::DropAllAttachments(ENetObjectAttachmentType Ty
 	ObjectToQueue.Remove(ObjectIndex);
 }
 
+SIZE_T FNetObjectAttachmentsWriter::GetReliableAttachmentCount(ENetObjectAttachmentType Type, uint32 ObjectIndex) const
+{
+	const FNetObjectAttachmentSendQueue* Queue = GetQueue(Type, ObjectIndex);
+	if (Queue == nullptr)
+	{
+		return 0;
+	}
+
+	return Queue->GetReliableCount();
+}
+
 SIZE_T FNetObjectAttachmentsWriter::GetUnreliableAttachmentCount(ENetObjectAttachmentType Type, uint32 ObjectIndex) const
 {
 	const FNetObjectAttachmentSendQueue* Queue = GetQueue(Type, ObjectIndex);
diff --git a/Engine/Source/Runtime/Net/Iris/Private/Iris/ReplicationSystem/AttachmentReplication.h b/Engine/Source/Runtime/Net/Iris/Private/Iris/ReplicationSystem/AttachmentReplication.h
--- a/Engine/Source/Runtime/Net/Iris/Private/Iris/ReplicationSystem/AttachmentReplication.h
+++ b/Engine/Source/Runtime/Net/Iris/Private/Iris/ReplicationSystem/AttachmentReplication.h
@@ -23,14 +23,19 @@ enum class ENetObjectAttachmentType : uint32
 	// For normal attachments a valid ObjectIndex is required.
 	Normal,
 
-	// For special attachments like OutOfBand and HugeObject the ObjectIndex is assumed to be zero.
+	// For special attachments like OutOfBand, HugeObject and DebugObjects the ObjectIndex is assumed to be zero.
 	OutOfBand,
 	HugeObject,
+	DebugObject,
 
 	// For internal use only
 	InternalCount
 };
 
+static constexpr uint32 NetObjectAttachmentTypeBitCount = 2;
+
+static_assert((uint32)ENetObjectAttachmentType::InternalCount == 1U << NetObjectAttachmentTypeBitCount, "ENetObjectAttachmentType must fit in NetObjectAttachementTypeBitCount bits");
+
 enum class EAttachmentWriteStatus : unsigned
 {
 	/** At least one attachment could be sent. */
@@ -73,6 +78,8 @@ public:
 
 	SIZE_T GetUnreliableCount() const;
 
+	SIZE_T GetReliableCount() const;
+
 	void DropUnreliable(bool &bOutHasUnsentAttachments);
 
 	bool HasUnsent() const;
@@ -136,6 +143,9 @@ public:
 
 	void DropAllAttachments(ENetObjectAttachmentType Type, uint32 ObjectIndex);
 
+	// Counts all Reliable attachments for the given type
+	SIZE_T GetReliableAttachmentCount(ENetObjectAttachmentType Type, uint32 ObjectIndex) const;
+
 	// Counts all unreliable attachments for the given type
 	SIZE_T GetUnreliableAttachmentCount(ENetObjectAttachmentType Type, uint32 ObjectIndex) const;
 
diff --git a/Engine/Source/Runtime/Net/Iris/Private/Iris/ReplicationSystem/ReplicationDataStream.cpp b/Engine/Source/Runtime/Net/Iris/Private/Iris/ReplicationSystem/ReplicationDataStream.cpp
--- a/Engine/Source/Runtime/Net/Iris/Private/Iris/ReplicationSystem/ReplicationDataStream.cpp
+++ b/Engine/Source/Runtime/Net/Iris/Private/Iris/ReplicationSystem/ReplicationDataStream.cpp
@@ -116,3 +116,13 @@ bool UReplicationDataStream::HasAcknowledgedAllReliableData() const
 
 	return true;
 }
+
+bool UReplicationDataStream::HasDebugDataToSend() const
+{
+	if (ReplicationWriter != nullptr)
+	{
+		return ReplicationWriter->HasDebugDataToSend();
+	}
+
+	return false;
+}
diff --git a/Engine/Source/Runtime/Net/Iris/Private/Iris/ReplicationSystem/ReplicationDataStream.h b/Engine/Source/Runtime/Net/Iris/Private/Iris/ReplicationSystem/ReplicationDataStream.h
--- a/Engine/Source/Runtime/Net/Iris/Private/Iris/ReplicationSystem/ReplicationDataStream.h
+++ b/Engine/Source/Runtime/Net/Iris/Private/Iris/ReplicationSystem/ReplicationDataStream.h
@@ -30,6 +30,7 @@ private:
 	virtual void ReadData(UE::Net::FNetSerializationContext& Context) override;
 	virtual void ProcessPacketDeliveryStatus(UE::Net::EPacketDeliveryStatus Status, FDataStreamRecord const* Record) override;
 	virtual bool HasAcknowledgedAllReliableData() const override;
+	virtual bool HasDebugDataToSend() const override;
 
 private:
 	UE::Net::Private::FReplicationReader* ReplicationReader;
diff --git a/Engine/Source/Runtime/Net/Iris/Private/Iris/ReplicationSystem/ReplicationReader.cpp b/Engine/Source/Runtime/Net/Iris/Private/Iris/ReplicationSystem/ReplicationReader.cpp
--- a/Engine/Source/Runtime/Net/Iris/Private/Iris/ReplicationSystem/ReplicationReader.cpp
+++ b/Engine/Source/Runtime/Net/Iris/Private/Iris/ReplicationSystem/ReplicationReader.cpp
@@ -1182,14 +1182,32 @@ uint32 FReplicationReader::ReadObjectBatch(FNetSerializationContext& Context, ui
 		PendingBatchData = UpdateUnresolvedMustBeMappedReferences(IncompleteHandle, TempMustBeMappedReferences, IrisAsyncLoadingPriority);
 	}
 	
-	// TODO: Test and support a case where the parents changed and now exist but we still have queued data waiting on the previous parents.
-	if (!bDoAllParentsExist)
+	// Normally we will only receive a single state/attachment update for an object in a single packet
+	// There are however a special case if we allow multiple HO-states to be in flight for the same object and they get assembled and dispatched during procssing of the same packet.
+	// In that case we want to dispatch them separately to maintain attachment/state application order.
+	bool bForcePendingBatchDataForHugeObject = false;
+	if (ReadObjectFlags & EReadObjectFlag::ReadObjectFlag_IsReadingHugeObjectBatch)
 	{
-		if (!PendingBatchData)
+		if (!ObjectsReadFromHugeObjects.Contains(IncompleteHandle))
+		{
+			ObjectsReadFromHugeObjects.Add(IncompleteHandle);
+		}
+		else
 		{
-			PendingBatchData = PendingBatchHolder.FindOrCreate(IncompleteHandle);
+			bForcePendingBatchDataForHugeObject = true;
 		}
+	}
+		
+	// Force create pending batch data if we have multiple incoming hugeobject batches for the same object or if we have unfulfilled creation dependencies.
+	const bool bForceCreatePendingBatchData = bForcePendingBatchDataForHugeObject || !bDoAllParentsExist;
+	if (bForceCreatePendingBatchData)
+	{
+		PendingBatchData = PendingBatchHolder.FindOrCreate(IncompleteHandle);
+	}
 
+	// TODO: Test and support a case where the parents changed and now exist but we still have queued data waiting on the previous parents.
+	if (!bDoAllParentsExist)
+	{
 		PendingBatchData->CreationDependentParents = CreationDependentParents;
 	}
 
@@ -1547,8 +1565,8 @@ void FReplicationReader::ReadObjectInBatch(FNetSerializationContext& Context, FN
 
 				if (!bHasErrors)
 				{
-					const bool bIsHugeObject = Reader.ReadBool();
-					AttachmentType = (bIsHugeObject ? ENetObjectAttachmentType::HugeObject : ENetObjectAttachmentType::OutOfBand);
+					const uint32 ReadAttachmentType = Reader.ReadBits(NetObjectAttachmentTypeBitCount);
+					AttachmentType = (ENetObjectAttachmentType)ReadAttachmentType;
 					bHasErrors = (!Parameters.bAllowReceivingAttachmentsFromRemoteObjectsNotInScope && AttachmentType == ENetObjectAttachmentType::OutOfBand);
 					UE_CLOG(bHasErrors, LogIris, Error, TEXT("FReplicationReader::ReadObject Bitstream corrupted. Reading OutOfBand attachment for object not in scope."));
 				}
@@ -1571,7 +1589,7 @@ void FReplicationReader::ReadObjectInBatch(FNetSerializationContext& Context, FN
 		}
 
 		// Fill in ReadObjectInfo, we must skip objects that has not been created and HugeObjects as they are not added to the dispatch list until they are fully assembled
-		const bool bShouldCommitPendingDispatchObjectInfo = (AttachmentType != ENetObjectAttachmentType::HugeObject) && !bIsReplicatedDestroyForInvalidObject;
+		const bool bShouldCommitPendingDispatchObjectInfo = !(AttachmentType == ENetObjectAttachmentType::HugeObject || AttachmentType == ENetObjectAttachmentType::DebugObject) && !bIsReplicatedDestroyForInvalidObject;
 		if (bShouldCommitPendingDispatchObjectInfo)
 		{
 			Info.InternalIndex = InternalIndex;
@@ -2716,7 +2734,8 @@ void FReplicationReader::ProcessHugeObjectAttachment(FNetSerializationContext& C
 		{
 			FNetBitStreamReader& Reader = *Context.GetBitStreamReader();
 			// Inject after all other trace events
-			FNetTrace::FoldTraceCollector(TraceCollector, HugeObjectTraceCollector, GetBitStreamPositionForNetTrace(Reader));
+			FNetTrace::FoldTraceCollector(TraceCollector, HugeObjectTraceCollector, GetBitStreamPositionForNetTrace(Reader) + DebugTraceOffset);
+			DebugTraceOffset += NetObjectBlob.GetRawDataBitCount();
 		}
 #endif
 	};
@@ -2729,6 +2748,7 @@ void FReplicationReader::ProcessHugeObjectAttachment(FNetSerializationContext& C
 	{
 		if (!Context.HasError())
 		{
+			UE_LOG(LogIris, Error, TEXT("Failed to process HugeObject payload"));
 			Context.SetError(GNetError_BitStreamError);
 			return;
 		}
@@ -3142,14 +3162,14 @@ void FReplicationReader::ProcessQueuedBatches()
 	}
 }
 
-void FReplicationReader::ProcessHugeObject(FNetSerializationContext& Context)
+void FReplicationReader::ProcessHugeObject(FNetSerializationContext& Context, ENetObjectAttachmentType AttachmentType)
 {
-	if (!Attachments.HasUnprocessedAttachments(ENetObjectAttachmentType::HugeObject, ObjectIndexForOOBAttachment))
+	if (!Attachments.HasUnprocessedAttachments(AttachmentType, ObjectIndexForOOBAttachment))
 	{
 		return;
 	}
 
-	FNetObjectAttachmentReceiveQueue* AttachmentQueue = Attachments.GetQueue(ENetObjectAttachmentType::HugeObject, ObjectIndexForOOBAttachment);
+	FNetObjectAttachmentReceiveQueue* AttachmentQueue = Attachments.GetQueue(AttachmentType, ObjectIndexForOOBAttachment);
 	while (const TRefCountPtr<FNetBlob>* Attachment = AttachmentQueue->PeekReliable())
 	{
 		ProcessHugeObjectAttachment(Context, *Attachment);
@@ -3239,12 +3259,23 @@ void FReplicationReader::Read(FNetSerializationContext& Context)
 	}
 
 	// Assemble and deserialize huge object if present
-	ProcessHugeObject(Context);
+	ProcessHugeObject(Context, ENetObjectAttachmentType::HugeObject);
 	if (Context.HasErrorOrOverflow())
 	{
 		return;
 	}
 
+	// Assemble and deserialize debug object if present
+	ProcessHugeObject(Context, ENetObjectAttachmentType::DebugObject);
+	if (Context.HasErrorOrOverflow())
+	{
+		return;
+	}
+
+	// Reset 
+	ObjectsReadFromHugeObjects.Reset();
+	DebugTraceOffset = 0U;
+
 	// Stats
 	UE_NET_TRACE_PACKET_STATSCOUNTER(Parameters.ReplicationSystem->GetId(), Parameters.ConnectionId, ReplicationReader.ReadObjectBatchCount, ReceivedObjectBatchCountToRead, ENetTraceVerbosity::Trace);
 	UE_NET_TRACE_PACKET_STATSCOUNTER(Parameters.ReplicationSystem->GetId(), Parameters.ConnectionId, ReplicationReader.ReadObjectsAndSubObjectsToDispatchCount, ObjectsToDispatchArray->Num(), ENetTraceVerbosity::Trace);
diff --git a/Engine/Source/Runtime/Net/Iris/Private/Iris/ReplicationSystem/ReplicationReader.h b/Engine/Source/Runtime/Net/Iris/Private/Iris/ReplicationSystem/ReplicationReader.h
--- a/Engine/Source/Runtime/Net/Iris/Private/Iris/ReplicationSystem/ReplicationReader.h
+++ b/Engine/Source/Runtime/Net/Iris/Private/Iris/ReplicationSystem/ReplicationReader.h
@@ -166,8 +166,8 @@ private:
 	// Process a single huge object attachment
 	void ProcessHugeObjectAttachment(FNetSerializationContext& Context, const TRefCountPtr<FNetBlob>& Attachment);
 
-	// Assemble and deserialize huge object if present
-	void ProcessHugeObject(FNetSerializationContext& Context);
+	// Assemble and deserialize huge object if present for the given attachmenttype.
+	void ProcessHugeObject(FNetSerializationContext& Context, ENetObjectAttachmentType AttachmentType);
 
 	// Resolve and dispatch unresolved references
 	void ResolveAndDispatchUnresolvedReferencesForObject(FNetSerializationContext& Context, uint32 InternalIndex);
@@ -291,6 +291,9 @@ private:
 	// Used during receive and processing of pending batches
 	TArray<FNetRefHandle> TempMustBeMappedReferences;
 
+	// Used during receive to track objects read from hugeobject path
+	TArray<FNetRefHandle> ObjectsReadFromHugeObjects;
+
 	// Preallocate the arrays used by BuildUnresolvedChangeMaskAndUpdateObjectReferenceTracking() to 
 	// avoid memory allocations during the frame.
 	FObjectReferenceTracker UnresolvedReferencesCache;
@@ -305,6 +308,9 @@ private:
 	IConsoleVariable const* DelayAttachmentsWithUnresolvedReferences;
 	uint32 NumHandlesPendingResolveLastUpdate = 0U;
 
+
+	uint32 DebugTraceOffset = 0U;
+
 	// Features that can aid in tracking down bitstream errors etc.
 	EReplicationDataStreamDebugFeatures StreamDebugFeatures = EReplicationDataStreamDebugFeatures::None;
 };
diff --git a/Engine/Source/Runtime/Net/Iris/Private/Iris/ReplicationSystem/ReplicationWriter.cpp b/Engine/Source/Runtime/Net/Iris/Private/Iris/ReplicationSystem/ReplicationWriter.cpp
--- a/Engine/Source/Runtime/Net/Iris/Private/Iris/ReplicationSystem/ReplicationWriter.cpp
+++ b/Engine/Source/Runtime/Net/Iris/Private/Iris/ReplicationSystem/ReplicationWriter.cpp
@@ -130,6 +130,7 @@ const TCHAR* FReplicationWriter::LexToString(const EReplicatedObjectState State)
 		TEXT("Invalid"),
 		TEXT("AttachmentToObjectNotInScope"),
 		TEXT("HugeObject"),
+		TEXT("DebugObject"),
 		TEXT("PendingCreate"),
 		TEXT("WaitOnCreateConfirmation"),
 		TEXT("Created"),
@@ -272,7 +273,7 @@ FReplicationWriter::~FReplicationWriter()
 	DiscardAllRecords();
 
 	// Freeing the huge object queue needs to be done before calling StopAllReplication() in order to be able to free any changemask allocations.
-	FreeHugeObjectSendQueue();
+	FreeHugeObjectSendQueues();
 	
 	StopAllReplication();
 }
@@ -349,6 +350,11 @@ bool FReplicationWriter::AreAllReliableAttachmentsSentAndAcked() const
 	return !bHasUnprocessedReliables && Attachments.AreAllObjectsReliableSentAndAcked();
 }
 
+bool FReplicationWriter::HasDebugDataToSend() const
+{
+	return Attachments.HasUnsentAttachments(ENetObjectAttachmentType::DebugObject, ObjectIndexForOOBAttachment);
+}
+
 void FReplicationWriter::SetState(uint32 InternalIndex, EReplicatedObjectState NewState)
 {
 	FReplicationInfo& Info = GetReplicationInfo(InternalIndex);
@@ -428,6 +434,12 @@ void FReplicationWriter::SetState(uint32 InternalIndex, EReplicatedObjectState N
 	Info.State = (uint32)NewState;
 }
 
+FReplicationWriter::FReplicationWriter()
+: HugeObjectSendQueue(TEXT("HugeObjectState"))
+, DebugObjectSendQueue(TEXT("DebugObjectState"), /*bAllowMultipleBatchesPerRootObjectInTransit*/true)
+{
+}
+
 void FReplicationWriter::Init(const FReplicationParameters& InParameters)
 {
 	// Store copy of parameters
@@ -558,6 +570,7 @@ void FReplicationWriter::StartReplication(uint32 InternalIndex)
 	Info.IsDeltaCompressionEnabled = BaselineManager->GetDeltaCompressionStatus(InternalIndex) == ENetObjectDeltaCompressionStatus::Allow ? 1U : 0U;
 	Info.LastAckedBaselineIndex = FDeltaCompressionBaselineManager::InvalidBaselineIndex;
 	Info.PendingBaselineIndex = FDeltaCompressionBaselineManager::InvalidBaselineIndex;
+	Info.IsDebugObject = Data.bIsDebugObject;
 
 	// Allocate storage for changemask (if needed)
 	FChangeMaskStorageOrPointer::Alloc(Info.ChangeMaskOrPtr, Info.ChangeMaskBitCount, s_DefaultChangeMaskAllocator);
@@ -1340,6 +1353,32 @@ void FReplicationWriter::HandleDeliveredRecord(const FReplicationRecord::FRecord
 		}
 		return;
 
+		case EReplicatedObjectState::DebugObject:
+		{
+			check(IsObjectIndexForOOBAttachment(InternalIndex));
+			Attachments.ProcessPacketDeliveryStatus(EPacketDeliveryStatus::Delivered, ENetObjectAttachmentType::DebugObject, ObjectIndexForOOBAttachment, AttachmentRecord);
+
+			DebugObjectSendQueue.AckObjects([this](const FHugeObjectContext& HugeObjectContext)
+			{
+				// If we've sent an entire huge objects we can ack everything in the payload and continue replicating this object using normal means.
+				const FReplicationInfo& HugeObjectReplicationInfo = this->GetReplicationInfo(HugeObjectContext.RootObjectInternalIndex);
+				for (const FObjectRecord& ObjectRecord : HugeObjectContext.BatchRecord.ObjectReplicationRecords)
+				{
+					FReplicationInfo& ReplicationInfo = this->GetReplicationInfo(ObjectRecord.Record.Index);
+					const uint32 ChangeMaskBitCount = ReplicationInfo.ChangeMaskBitCount;
+					this->HandleDeliveredRecord(ObjectRecord.Record, ReplicationInfo, ObjectRecord.AttachmentRecord);
+					if (ObjectRecord.Record.HasChangeMask)
+					{
+						FChangeMaskStorageOrPointer::Free(ObjectRecord.Record.ChangeMaskOrPtr, ChangeMaskBitCount, s_DefaultChangeMaskAllocator);
+					}
+				}
+
+				// We need to explicitly acknowledge exports made through the huge object batch
+				this->NetExports->AcknowledgeBatchExports(HugeObjectContext.BatchExports);
+			});
+		}
+		return;
+
 		default:
 		break;
 	}
@@ -1473,6 +1512,31 @@ void FReplicationWriter::HandleDiscardedRecord(const FReplicationRecord::FRecord
 			});
 		}
 		return;
+
+		case EReplicatedObjectState::DebugObject:
+		{
+			check(IsObjectIndexForOOBAttachment(InternalIndex));
+
+			// Deal with it similar to if the entire state has been sent as we need to go through all records.
+			Attachments.ProcessPacketDeliveryStatus(EPacketDeliveryStatus::Discard, ENetObjectAttachmentType::DebugObject, ObjectIndexForOOBAttachment, AttachmentRecord);
+
+			DebugObjectSendQueue.AckObjects([this](const FHugeObjectContext& HugeObjectContext)
+			{
+				const FReplicationInfo& HugeObjectReplicationInfo = this->GetReplicationInfo(HugeObjectContext.RootObjectInternalIndex);
+				for (const FObjectRecord& ObjectRecord : HugeObjectContext.BatchRecord.ObjectReplicationRecords)
+				{
+					FReplicationInfo& ReplicationInfo = this->GetReplicationInfo(ObjectRecord.Record.Index);
+					const uint32 ChangeMaskBitCount = ReplicationInfo.ChangeMaskBitCount;
+					this->HandleDiscardedRecord(ObjectRecord.Record, ReplicationInfo, ObjectRecord.AttachmentRecord);
+					if (ObjectRecord.Record.HasChangeMask)
+					{
+						FChangeMaskStorageOrPointer::Free(ObjectRecord.Record.ChangeMaskOrPtr, ChangeMaskBitCount, s_DefaultChangeMaskAllocator);
+					}
+				}
+			});
+		}
+		return;
+
 	}
 }
 
@@ -1815,6 +1879,13 @@ void FReplicationWriter::HandleDroppedRecord(const FReplicationRecord::FRecordIn
 		}
 		return;
 
+		case EReplicatedObjectState::DebugObject:
+		{
+			check(IsObjectIndexForOOBAttachment(InternalIndex));
+			Attachments.ProcessPacketDeliveryStatus(EPacketDeliveryStatus::Lost, ENetObjectAttachmentType::DebugObject, ObjectIndexForOOBAttachment, AttachmentRecord);
+		}
+		return;
+
 		
 		default:
 		break;
@@ -1879,7 +1950,20 @@ void FReplicationWriter::CreateObjectRecord(const FNetBitArrayView* ChangeMask,
 	FReplicationRecord::FRecordInfo& RecordInfo = OutRecord.Record;
 
 	RecordInfo.Index = ObjectInfo.InternalIndex;
-	RecordInfo.ReplicatedObjectState = ObjectInfo.AttachmentType == ENetObjectAttachmentType::HugeObject ? uint8(EReplicatedObjectState::HugeObject) : (uint8)Info.GetState();
+
+	if (ObjectInfo.AttachmentType == ENetObjectAttachmentType::HugeObject)
+	{
+		RecordInfo.ReplicatedObjectState = uint8(EReplicatedObjectState::HugeObject);
+	}
+	else if (ObjectInfo.AttachmentType == ENetObjectAttachmentType::DebugObject)
+	{
+		RecordInfo.ReplicatedObjectState = uint8(EReplicatedObjectState::DebugObject);
+	}
+	else
+	{
+		RecordInfo.ReplicatedObjectState = (uint8)Info.GetState();
+	}
+
 	RecordInfo.HasChangeMask = ChangeMask ? 1U : 0U;
 	RecordInfo.HasAttachments = (OutRecord.AttachmentRecord.IsValid() ? 1U : 0U);
 	RecordInfo.WroteTearOff = ObjectInfo.bSentTearOff;
@@ -2347,7 +2431,7 @@ void FReplicationWriter::SerializeObjectStateDelta(FNetSerializationContext& Con
 	}
 }
 
-FReplicationWriter::EWriteObjectStatus FReplicationWriter::WriteObjectAndSubObjects(FNetSerializationContext& Context, uint32 InternalIndex, uint32 WriteObjectFlags, FBatchInfo& OutBatchInfo)
+FReplicationWriter::EWriteObjectStatus FReplicationWriter::WriteObjectAndSubObjects(FNetSerializationContext& Context, uint32 InternalIndex, FWriteObjectTraits WriteObjectTraits, FBatchInfo& OutBatchInfo)
 {
 	FNetBitStreamWriter& Writer = *Context.GetBitStreamWriter();
 	
@@ -2370,12 +2454,17 @@ FReplicationWriter::EWriteObjectStatus FReplicationWriter::WriteObjectAndSubObje
 
 #if UE_NET_TRACE_ENABLED
 	FNetRefHandle NetRefHandleForTraceScope = NetRefHandle;
-	if (WriteObjectFlags & EWriteObjectFlag::WriteObjectFlag_HugeObject)
+	if (InternalIndex == ObjectIndexForOOBAttachment)
 	{
-		const FInternalNetRefIndex HugeObjectInternalIndex = HugeObjectSendQueue.GetRootObjectInternalIndexForTrace();
-		if (HugeObjectInternalIndex != FNetRefHandleManager::InvalidInternalIndex)
+		const ENetObjectAttachmentType OOBWriteAttachmentType = (ENetObjectAttachmentType)WriteObjectTraits.OOBWriteAttachmentType;
+		if (OOBWriteAttachmentType == ENetObjectAttachmentType::HugeObject || OOBWriteAttachmentType == ENetObjectAttachmentType::DebugObject)
 		{
-			NetRefHandleForTraceScope = NetRefHandleManager->GetReplicatedObjectDataNoCheck(HugeObjectInternalIndex).RefHandle;
+			const FHugeObjectSendQueue& HugeObjectSendQueueForDebug = OOBWriteAttachmentType == ENetObjectAttachmentType::HugeObject ? HugeObjectSendQueue : DebugObjectSendQueue;
+			const FInternalNetRefIndex HugeObjectInternalIndex = HugeObjectSendQueueForDebug.GetRootObjectInternalIndexForTrace();
+			if (HugeObjectInternalIndex != FNetRefHandleManager::InvalidInternalIndex)
+			{
+				NetRefHandleForTraceScope = NetRefHandleManager->GetReplicatedObjectDataNoCheck(HugeObjectInternalIndex).RefHandle;
+			}
 		}
 	}
 	UE_NET_TRACE_NAMED_OBJECT_SCOPE(ObjectTraceScope, NetRefHandleForTraceScope, Writer, Context.GetTraceCollector(), ENetTraceVerbosity::Trace);
@@ -2384,7 +2473,7 @@ FReplicationWriter::EWriteObjectStatus FReplicationWriter::WriteObjectAndSubObje
 	// We only need to write batch header for the root object
 	const bool bWriteBatchHeader = !Info.IsSubObject;
 	uint32 InitialStateHeaderPos = 0U;
-	const uint32 NumBitsUsedForBatchSize = (WriteObjectFlags & EWriteObjectFlag::WriteObjectFlag_IsWritingHugeObjectBatch) == 0U ? Parameters.NumBitsUsedForBatchSize : Parameters.NumBitsUsedForHugeObjectBatchSize;
+	const uint32 NumBitsUsedForBatchSize = WriteObjectTraits.bIsWritingHugeObjectBatch == 0U ? Parameters.NumBitsUsedForBatchSize : Parameters.NumBitsUsedForHugeObjectBatchSize;
 
 	// This is the beginning of what we treat as a batch on the receiving end
 	if (bWriteBatchHeader)
@@ -2520,9 +2609,9 @@ FReplicationWriter::EWriteObjectStatus FReplicationWriter::WriteObjectAndSubObje
 
 	const bool bAllowStateAndAttachmentSending = (!Info.SubObjectPendingDestroy || bAllowSendStateWithEndReplicationForSubObjects);
 	const bool bIsObjectIndexForAttachment = IsObjectIndexForOOBAttachment(InternalIndex);
-	const bool bHasState = (bIsInitialState || Info.HasDirtyChangeMask) && !!(WriteObjectFlags & EWriteObjectFlag::WriteObjectFlag_State) && bAllowStateAndAttachmentSending;
+	const bool bHasState = (bIsInitialState || Info.HasDirtyChangeMask) && !!WriteObjectTraits.bWriteState && bAllowStateAndAttachmentSending;
 	const bool bHasAttachments = (Info.HasAttachments || bIsObjectIndexForAttachment) && bAllowStateAndAttachmentSending;
-	const bool bWriteAttachments = bHasAttachments && !!(WriteObjectFlags & EWriteObjectFlag::WriteObjectFlag_Attachments) && bAllowStateAndAttachmentSending;
+	const bool bWriteAttachments = bHasAttachments && !!(WriteObjectTraits.bWriteAttachments) && bAllowStateAndAttachmentSending;
 	BatchEntry.bHasUnsentAttachments = bHasAttachments;
 
 	// Check if we must defer tearoff until after flush
@@ -2705,8 +2794,9 @@ FReplicationWriter::EWriteObjectStatus FReplicationWriter::WriteObjectAndSubObje
 				BatchEntry.AttachmentType = ENetObjectAttachmentType::Normal;
 				if (bIsObjectIndexForAttachment)
 				{
-					BatchEntry.AttachmentType = (WriteObjectFlags & EWriteObjectFlag::WriteObjectFlag_HugeObject ? ENetObjectAttachmentType::HugeObject : ENetObjectAttachmentType::OutOfBand);
-					AttachmentWriter.WriteBool(BatchEntry.AttachmentType == ENetObjectAttachmentType::HugeObject);
+					check((ENetObjectAttachmentType)WriteObjectTraits.OOBWriteAttachmentType != ENetObjectAttachmentType::Normal);
+					BatchEntry.AttachmentType = (ENetObjectAttachmentType)WriteObjectTraits.OOBWriteAttachmentType;
+					AttachmentWriter.WriteBits(WriteObjectTraits.OOBWriteAttachmentType, NetObjectAttachmentTypeBitCount);
 				}
 
 				const EAttachmentWriteStatus AttachmentWriteStatus = Attachments.Serialize(AttachmentContext, BatchEntry.AttachmentType, InternalIndex, NetRefHandle, BatchEntry.AttachmentRecord, BatchEntry.bHasUnsentAttachments);
@@ -2722,6 +2812,18 @@ FReplicationWriter::EWriteObjectStatus FReplicationWriter::WriteObjectAndSubObje
 						HugeObjectSendQueue.Stats.StartStallTime = 0;
 					}
 				}
+				else if (BatchEntry.AttachmentType == ENetObjectAttachmentType::DebugObject)
+				{
+					if (AttachmentWriteStatus == EAttachmentWriteStatus::ReliableWindowFull)
+					{
+						DebugObjectSendQueue.Stats.StartStallTime = FPlatformTime::Cycles64();
+					}
+					else
+					{
+						// Clear stall time now that we were theoretically able to send something.
+						DebugObjectSendQueue.Stats.StartStallTime = 0;
+					}			
+				}
 
 				// If we didn't manage to fit any attachments then clear the HasAttachments bool in the packet
 				if (AttachmentWriter.GetPosBits() == 0 || AttachmentWriter.IsOverflown())
@@ -2740,10 +2842,10 @@ FReplicationWriter::EWriteObjectStatus FReplicationWriter::WriteObjectAndSubObje
 					
 					// If we should have had enough space to write a attachment, the attachment + exports must be huge and we need to fall back on using the huge object path
 					const uint32 SplitThreshold = PartialNetObjectAttachmentHandler->GetConfig()->GetBitCountSplitThreshold() * 2;
-					const bool bFallbackToHugeObjectPath = BatchEntry.AttachmentType != ENetObjectAttachmentType::HugeObject && BatchEntry.bHasUnsentAttachments && (BitsThatWasAvailableForAttachements >= SplitThreshold);
+					const bool bFallbackToHugeObjectPath = !(BatchEntry.AttachmentType == ENetObjectAttachmentType::HugeObject || BatchEntry.AttachmentType == ENetObjectAttachmentType::DebugObject) && BatchEntry.bHasUnsentAttachments && (BitsThatWasAvailableForAttachements >= SplitThreshold);
 					if (bFallbackToHugeObjectPath)
 					{
-						UE_LOG(LogIris, Verbose, TEXT("Failed to write huge attachment for object %s ( InternalIndex: %u ), forcing fallback on hugeobject for attachments"), *NetRefHandle.ToString(), InternalIndex);
+						UE_LOG(LogIris, Verbose, TEXT("Failed to write large attachment for object %s ( InternalIndex: %u ), forcing fallback on hugeobject for attachments"), *NetRefHandle.ToString(), InternalIndex);
 						Writer.DoOverflow();
 					}
 					else if (!(BatchEntry.bSentState || bSentTearOff || Info.SubObjectPendingDestroy || Info.HasDirtySubObjects || bWriteBatchHeader))
@@ -2861,7 +2963,7 @@ FReplicationWriter::EWriteObjectStatus FReplicationWriter::WriteObjectAndSubObje
 			}
 
 			const int BatchObjectInfoCount = OutBatchInfo.ObjectInfos.Num();
-			EWriteObjectStatus SubObjectWriteStatus = WriteObjectAndSubObjects(Context, SubObjectInternalIndex, WriteObjectFlags, OutBatchInfo);
+			EWriteObjectStatus SubObjectWriteStatus = WriteObjectAndSubObjects(Context, SubObjectInternalIndex, WriteObjectTraits, OutBatchInfo);
 			if (!IsWriteObjectSuccess(SubObjectWriteStatus))
 			{
 				// SubObject will rollback on fail (and report its own waste) but we as we will rollback successfully written subobjects it is better to at least report it with the owner.
@@ -2944,7 +3046,7 @@ FReplicationWriter::EWriteObjectStatus FReplicationWriter::WriteObjectAndSubObje
 	return EWriteObjectStatus::Success;
 }
 
-FReplicationWriter::EWriteObjectStatus FReplicationWriter::WriteObjectInBatch(FNetSerializationContext& Context, uint32 InternalIndex, uint32 WriteObjectFlags, FBatchInfo& OutBatchInfo)
+FReplicationWriter::EWriteObjectStatus FReplicationWriter::WriteObjectInBatch(FNetSerializationContext& Context, uint32 InternalIndex, FWriteObjectTraits WriteObjectTraits, FBatchInfo& OutBatchInfo)
 {
 	{
 		UE_NET_IRIS_STATS_TIMER(Timer, Context.GetNetStatsContext());
@@ -2958,7 +3060,7 @@ FReplicationWriter::EWriteObjectStatus FReplicationWriter::WriteObjectInBatch(FN
 		}
 
 		// Write parent object and subobjects
-		const EWriteObjectStatus WriteObjectStatus = WriteObjectAndSubObjects(Context, InternalIndex, WriteObjectFlags, OutBatchInfo);
+		const EWriteObjectStatus WriteObjectStatus = WriteObjectAndSubObjects(Context, InternalIndex, WriteObjectTraits, OutBatchInfo);
 		if (!IsWriteObjectSuccess(WriteObjectStatus))
 		{
 			UE_NET_IRIS_STATS_ADD_TIME_AND_COUNT_FOR_OBJECT_AS_WASTE(Timer, Write, InternalIndex);
@@ -2980,7 +3082,7 @@ FReplicationWriter::EWriteObjectStatus FReplicationWriter::WriteObjectInBatch(FN
 				UE_NET_TRACE_SCOPE(DependentObjectData, *Context.GetBitStreamWriter(), Context.GetTraceCollector(), ENetTraceVerbosity::VeryVerbose);
 
 				OutBatchInfo.ParentInternalIndex = DependentInternalIndex;
-				EWriteObjectStatus DependentObjectWriteStatus = WriteObjectInBatch(Context, DependentInternalIndex, WriteObjectFlags, OutBatchInfo);
+				EWriteObjectStatus DependentObjectWriteStatus = WriteObjectInBatch(Context, DependentInternalIndex, WriteObjectTraits, OutBatchInfo);
 				if (!IsWriteObjectSuccess(DependentObjectWriteStatus))
 				{
 					// Restore ParentInternalIndex
@@ -2997,14 +3099,17 @@ FReplicationWriter::EWriteObjectStatus FReplicationWriter::WriteObjectInBatch(FN
 	return EWriteObjectStatus::Success;
 }
 
-FReplicationWriter::EWriteStatus FReplicationWriter::PrepareAndSendHugeObjectPayload(FNetSerializationContext& Context, FInternalNetRefIndex InternalIndex)
+FReplicationWriter::EWriteStatus FReplicationWriter::PrepareAndSendHugeObjectPayload(FNetSerializationContext& Context, FInternalNetRefIndex InternalIndex, ENetObjectAttachmentType AttachmentType)
 {
 	IRIS_PROFILER_SCOPE(FReplicationWriter_PrepareAndSendHugeObjectPayload);
 
+	const bool bIsPreparingDebugObject = AttachmentType == ENetObjectAttachmentType::DebugObject;
+	FHugeObjectSendQueue& InHugeObjectSendQueue = bIsPreparingDebugObject ?  DebugObjectSendQueue : HugeObjectSendQueue;
+
 	// Sanity check
-	if (HugeObjectSendQueue.IsFull() || HugeObjectSendQueue.IsObjectInQueue(InternalIndex, false))
+	if (InHugeObjectSendQueue.IsFull() || (!bIsPreparingDebugObject && InHugeObjectSendQueue.IsObjectInQueue(InternalIndex, false)))
 	{
-		ensureMsgf(false, TEXT("HugeObjectSendQueue should not be full or already transmitting: %s"), *NetRefHandleManager->PrintObjectFromIndex(InternalIndex));
+		ensureMsgf(false, TEXT("PrepareAndSendHugeObjectPayload (Huge|Debug)ObjectSendQueue should not be full or already transmitting: %s"), *NetRefHandleManager->PrintObjectFromIndex(InternalIndex));
 		return EWriteStatus::Skipped;
 	}
 
@@ -3021,36 +3126,39 @@ FReplicationWriter::EWriteStatus FReplicationWriter::PrepareAndSendHugeObjectPay
 	FNetSerializationContext HugeObjectSerializationContext = Context.MakeSubContext(&HugeObjectWriter);
 
 #if UE_NET_TRACE_ENABLED
-	if (!HugeObjectSendQueue.TraceCollector)
+	if (!InHugeObjectSendQueue.TraceCollector)
 	{
-		HugeObjectSendQueue.TraceCollector = UE_NET_TRACE_CREATE_COLLECTOR(ENetTraceVerbosity::Trace);
+		InHugeObjectSendQueue.TraceCollector = UE_NET_TRACE_CREATE_COLLECTOR(ENetTraceVerbosity::Trace);
 	}
 	else
 	{
-		HugeObjectSendQueue.TraceCollector->Reset();
+		InHugeObjectSendQueue.TraceCollector->Reset();
 	}
-	HugeObjectSerializationContext.SetTraceCollector(HugeObjectSendQueue.TraceCollector);
+	HugeObjectSerializationContext.SetTraceCollector(InHugeObjectSendQueue.TraceCollector);
 #endif
 
 	// Huge object header needed for the receiving side to be able to process this correctly.
 	FNetObjectBlob::FHeader HugeObjectHeader = {};
 	const uint32 HeaderPos = HugeObjectWriter.GetPosBits();
 	FNetObjectBlob::SerializeHeader(HugeObjectSerializationContext, HugeObjectHeader);
-	const uint32 PastHeaderPos = HugeObjectWriter.GetPosBits();
 
 	// As the huge object will most likely be processed out of order we need to write the stream debug features again
 	WriteStreamDebugFeatures(HugeObjectSerializationContext);
 
+	const uint32 PastHeaderPos = HugeObjectWriter.GetPosBits();
+
 	FHugeObjectContext HugeObjectContext;
 
 	FBatchInfo BatchInfo;
 	BatchInfo.Type = EBatchInfoType::Internal;
 	BatchInfo.ParentInternalIndex = InternalIndex;
-	uint32 WriteObjectFlags = WriteObjectFlag_State | WriteObjectFlag_IsWritingHugeObjectBatch;
-	// Get the creation going as quickly as possible.
-	if (!Context.IsInitState())
+	FWriteObjectTraits WriteObjectTraits({.bWriteState = 1U, .OOBWriteAttachmentType = (uint32)ENetObjectAttachmentType::OutOfBand, .bIsWritingHugeObjectBatch = 1U});
+
+	// Get the creation going as quickly as possible for hugeobjects by not including attachments. For DebugObjects we include everything.
+	const bool bIncludeAttachements = bIsPreparingDebugObject || !Context.IsInitState();
+	if (bIncludeAttachements)
 	{
-		WriteObjectFlags |= WriteObjectFlag_Attachments;
+		WriteObjectTraits.bWriteAttachments = 1U;
 	}
 	
 	// Push new ExportContext for the hugeobject-batch as we cannot share exports with an OOB object
@@ -3061,7 +3169,7 @@ FReplicationWriter::EWriteStatus FReplicationWriter::PrepareAndSendHugeObjectPay
 		UE_NET_TRACE_SCOPE(HugeObjectState, *HugeObjectSerializationContext.GetBitStreamWriter(), HugeObjectSerializationContext.GetTraceCollector(), ENetTraceVerbosity::Trace);
 
 		// We can encounter other errors than bitstream overflow now that we've got a really large buffer to write to.
-		const EWriteObjectStatus WriteHugeObjectStatus = WriteObjectInBatch(HugeObjectSerializationContext, InternalIndex, WriteObjectFlags, BatchInfo);
+		const EWriteObjectStatus WriteHugeObjectStatus = WriteObjectInBatch(HugeObjectSerializationContext, InternalIndex, WriteObjectTraits, BatchInfo);
 
 		// If we cannot fit the object in the largest supported buffer then we will never fit the object.
 		if (WriteHugeObjectStatus == EWriteObjectStatus::BitStreamOverflow)
@@ -3079,8 +3187,9 @@ FReplicationWriter::EWriteStatus FReplicationWriter::PrepareAndSendHugeObjectPay
 			return EWriteStatus::Abort;
 		}
 
-		// If we encounter some other error we can try sending a smaller object in the meantime.
-		if (!IsWriteObjectSuccess(WriteHugeObjectStatus))
+		const bool bWroteData = HugeObjectWriter.GetPosBits() > PastHeaderPos;
+		// If we encounter some other error or did end up not writing any data due to full attachment queue etc. we skip the huge object and try to write other data instead.
+		if (!IsWriteObjectSuccess(WriteHugeObjectStatus) || !bWroteData)
 		{
 			// Cleanup data from batch
 			HandleObjectBatchFailure(WriteHugeObjectStatus, BatchInfo, WriteBitStreamInfo);
@@ -3090,9 +3199,9 @@ FReplicationWriter::EWriteStatus FReplicationWriter::PrepareAndSendHugeObjectPay
 		}
 	}
 
-	if (HugeObjectSendQueue.IsEmpty())
+	if (InHugeObjectSendQueue.IsEmpty())
 	{
-		HugeObjectSendQueue.Stats.StartSendingTime = FPlatformTime::Cycles64();
+		InHugeObjectSendQueue.Stats.StartSendingTime = FPlatformTime::Cycles64();
 	}
 
 	HugeObjectContext.RootObjectInternalIndex = InternalIndex;
@@ -3119,7 +3228,7 @@ FReplicationWriter::EWriteStatus FReplicationWriter::PrepareAndSendHugeObjectPay
 	TArrayView<HugeObjectStorageType> PayloadView(HugeObjectPayload.GetData(), StorageWordsWritten);
 	TRefCountPtr<FNetObjectBlob> NetObjectBlob = NetObjectBlobHandler->CreateNetObjectBlob(PayloadView, PayLoadBitCount);
 	TArray<TRefCountPtr<FNetBlob>> PartialNetBlobs;
-	const bool bSplitSuccess = PartialNetObjectAttachmentHandler->SplitRawDataNetBlob(TRefCountPtr<FRawDataNetBlob>(NetObjectBlob.GetReference()), PartialNetBlobs, HugeObjectSendQueue.DebugName);
+	const bool bSplitSuccess = PartialNetObjectAttachmentHandler->SplitRawDataNetBlob(TRefCountPtr<FRawDataNetBlob>(NetObjectBlob.GetReference()), PartialNetBlobs, InHugeObjectSendQueue.DebugName);
 	if (!bSplitSuccess)
 	{
 		UE_LOG(LogIris, Error, TEXT("Unable to split huge object %s payload. Connection %u will be disconnected."), *NetRefHandleManager->PrintObjectFromIndex(InternalIndex), Context.GetLocalConnectionId());
@@ -3127,8 +3236,7 @@ FReplicationWriter::EWriteStatus FReplicationWriter::PrepareAndSendHugeObjectPay
 		return EWriteStatus::Abort;
 	}
 
-	// Enqueue attachments
-	const bool bEnqueueSuccess = Attachments.Enqueue(ENetObjectAttachmentType::HugeObject, ObjectIndexForOOBAttachment, MakeArrayView(PartialNetBlobs.GetData(), PartialNetBlobs.Num()));
+	const bool bEnqueueSuccess = Attachments.Enqueue(AttachmentType, ObjectIndexForOOBAttachment, MakeArrayView(PartialNetBlobs.GetData(), PartialNetBlobs.Num()));
 	check(bEnqueueSuccess);
 	if (!bEnqueueSuccess)
 	{
@@ -3139,7 +3247,7 @@ FReplicationWriter::EWriteStatus FReplicationWriter::PrepareAndSendHugeObjectPay
 
 	// Add huge object to queue
 	HugeObjectContext.Blobs = MoveTemp(PartialNetBlobs);
-	const bool bHugeObjectWasEnqueued = HugeObjectSendQueue.EnqueueHugeObject(HugeObjectContext);
+	const bool bHugeObjectWasEnqueued = InHugeObjectSendQueue.EnqueueHugeObject(HugeObjectContext);
 	check(bHugeObjectWasEnqueued);
 	if (!bHugeObjectWasEnqueued)
 	{
@@ -3149,6 +3257,12 @@ FReplicationWriter::EWriteStatus FReplicationWriter::PrepareAndSendHugeObjectPay
 	}
 
 	// Write huge object attachment(s)
+	if (bIsPreparingDebugObject)
+	{
+		// DebugObject attachments are sent separately to be able to not count them against bandwidth limits.
+		return EWriteStatus::Skipped;
+	}
+	else
 	{
 		FNetBitStreamWriter& Writer = *Context.GetBitStreamWriter();
 		UE_NET_TRACE_SCOPE(Batch, Writer, Context.GetTraceCollector(), ENetTraceVerbosity::Trace);
@@ -3157,8 +3271,9 @@ FReplicationWriter::EWriteStatus FReplicationWriter::PrepareAndSendHugeObjectPay
 		FBatchInfo HugeObjectBatchInfo;
 		HugeObjectBatchInfo.Type = EBatchInfoType::HugeObject;
 		HugeObjectBatchInfo.ParentInternalIndex = FNetRefHandleManager::InvalidInternalIndex;
-		const uint32 WriteHugeObjectFlags = WriteObjectFlag_Attachments | WriteObjectFlag_HugeObject;
-		const EWriteObjectStatus HugeObjectStatus = WriteObjectInBatch(Context, ObjectIndexForOOBAttachment, WriteHugeObjectFlags, HugeObjectBatchInfo);
+
+		const FWriteObjectTraits WriteHugeObjectInfo({.bWriteAttachments = 1U, .OOBWriteAttachmentType = (uint32)ENetObjectAttachmentType::HugeObject});
+		const EWriteObjectStatus HugeObjectStatus = WriteObjectInBatch(Context, ObjectIndexForOOBAttachment, WriteHugeObjectInfo, HugeObjectBatchInfo);
 		const uint32 BitsWrittenByHugeObjectAttachments = Writer.GetPosBits() - RollbackScope.GetStartPos();
 		if (!IsWriteObjectSuccess(HugeObjectStatus) || (BitsWrittenByHugeObjectAttachments == 0U))
 		{
@@ -3185,22 +3300,21 @@ FReplicationWriter::EWriteStatus FReplicationWriter::PrepareAndSendHugeObjectPay
 		WriteContext.bHasHugeObjectToSend = bHasHugeObjectToSend;
 		if (!bHasHugeObjectToSend)
 		{
-			HugeObjectSendQueue.Stats.EndSendingTime = FPlatformTime::Cycles64();
+			InHugeObjectSendQueue.Stats.EndSendingTime = FPlatformTime::Cycles64();
 		}
 
 		return EWriteStatus::Written;
 	}
-
-	// Should not get here.
-	return EWriteStatus::Abort;
 }
 
-FReplicationWriter::FWriteBatchResult FReplicationWriter::WriteObjectBatch(FNetSerializationContext& Context, FInternalNetRefIndex InternalIndex, uint32 WriteObjectFlags)
+FReplicationWriter::FWriteBatchResult FReplicationWriter::WriteObjectBatch(FNetSerializationContext& Context, FInternalNetRefIndex InternalIndex, FWriteObjectTraits WriteObjectTraits)
 {
 	IRIS_PROFILER_SCOPE(FReplicationWriter_WriteObjectBatch);
 
+	FReplicationInfo& ReplicationInfo = GetReplicationInfo(InternalIndex);
+
 	// If this is a destruction info we treat it differently and just write the information required to destruct the object
-	if (GetReplicationInfo(InternalIndex).IsDestructionInfo)
+	if (ReplicationInfo.IsDestructionInfo)
 	{
 		const EWriteStatus Status = WriteDestructionInfo(Context, InternalIndex);
 		return FWriteBatchResult
@@ -3212,6 +3326,24 @@ FReplicationWriter::FWriteBatchResult FReplicationWriter::WriteObjectBatch(FNetS
 
 	FNetBitStreamWriter& Writer = *Context.GetBitStreamWriter();
 
+	// We always send debug objects using the a special variant of the hugeobjectpath using a separate queue for written data.
+	if (ReplicationInfo.IsDebugObject)
+	{
+		const EWriteStatus SendDebugObjectStatus = PrepareAndSendHugeObjectPayload(Context, InternalIndex, ENetObjectAttachmentType::DebugObject);
+
+		// If the huge object wrote data it will be tracked as a single batch
+		if (SendDebugObjectStatus == EWriteStatus::Written)
+		{
+			++WriteContext.WrittenBatchCount;
+		}
+				
+		return FWriteBatchResult
+		{
+			.Status = SendDebugObjectStatus,
+			.NumWritten = (SendDebugObjectStatus == EWriteStatus::Written) ? 1U : 0U,
+		};
+	}
+
 	// Batch successful writes and commit them as a atomic batch.
 	// It is a fail if we fail to write any subobject with dirty state.
 	// It is also not ok to skip over creation header- if we do then the entire batch needs to be delayed.
@@ -3224,11 +3356,29 @@ FReplicationWriter::FWriteBatchResult FReplicationWriter::WriteObjectBatch(FNetS
 
 		WriteBitStreamInfo.BatchStartPos = Writer.GetPosBits();
 		FBatchInfo BatchInfo;
-		BatchInfo.Type = (InternalIndex == ObjectIndexForOOBAttachment ? (WriteObjectFlags & WriteObjectFlag_HugeObject ? EBatchInfoType::HugeObject : EBatchInfoType::OOBAttachment) : EBatchInfoType::Object);
+		if (InternalIndex == ObjectIndexForOOBAttachment)
+		{
+			if ((ENetObjectAttachmentType)WriteObjectTraits.OOBWriteAttachmentType == ENetObjectAttachmentType::HugeObject)
+			{
+				BatchInfo.Type = EBatchInfoType::HugeObject;	
+			}
+			else if ((ENetObjectAttachmentType)WriteObjectTraits.OOBWriteAttachmentType == ENetObjectAttachmentType::DebugObject)
+			{
+				BatchInfo.Type = EBatchInfoType::DebugObject;	
+			}
+			else
+			{
+				BatchInfo.Type = EBatchInfoType::OOBAttachment;
+			}
+		}
+		else
+		{
+			BatchInfo.Type = EBatchInfoType::Object;
+		}
 		BatchInfo.ParentInternalIndex = InternalIndex;
 
 		// Write an object and its subobjects. If object has dependent objects pending creation we currently write them as well as an individual batch.
-		const EWriteObjectStatus WriteObjectStatus = WriteObjectInBatch(Context, InternalIndex, WriteObjectFlags, BatchInfo);
+		const EWriteObjectStatus WriteObjectStatus = WriteObjectInBatch(Context, InternalIndex, WriteObjectTraits, BatchInfo);
 
 		if (IsWriteObjectSuccess(WriteObjectStatus))
 		{
@@ -3265,7 +3415,7 @@ FReplicationWriter::FWriteBatchResult FReplicationWriter::WriteObjectBatch(FNetS
 			case EWriteObjectRetryMode::TrySmallObject:
 			{
 				++WriteContext.FailedToWriteSmallObjectCount;
-				return FWriteBatchResult{ .Status = EWriteStatus::Skipped };
+				return FWriteBatchResult{ .Status = EWriteStatus::SkippedRetryNextPacket };
 			}
 
 			case EWriteObjectRetryMode::SplitHugeObject:
@@ -3282,7 +3432,7 @@ FReplicationWriter::FWriteBatchResult FReplicationWriter::WriteObjectBatch(FNetS
 
 	// Try #2 - Object will be serialized to a temporary buffer of maximum supported size and split into multiple chunks.
 	{
-		const EWriteStatus SendHugeObjectStatus = PrepareAndSendHugeObjectPayload(Context, InternalIndex);
+		const EWriteStatus SendHugeObjectStatus = PrepareAndSendHugeObjectPayload(Context, InternalIndex, ENetObjectAttachmentType::HugeObject);
 
 		// If the huge object wrote data it will be tracked as a single batch
 		if (SendHugeObjectStatus == EWriteStatus::Written)
@@ -3409,7 +3559,7 @@ uint32 FReplicationWriter::WriteOOBAttachments(FNetSerializationContext& Context
 		if (WriteContext.bHasOOBAttachmentsToSend && CanSendObject(ObjectIndexForOOBAttachment) && (ReplicationRecord.GetUnusedInfoCount() > 0U))
 		{
 			IRIS_PROFILER_SCOPE(FReplicationWriter_WriteOOBAttachments);
-			const FWriteBatchResult Result = WriteObjectBatch(Context, ObjectIndexForOOBAttachment, WriteObjectFlag_Attachments);
+			const FWriteBatchResult Result = WriteObjectBatch(Context, ObjectIndexForOOBAttachment, FWriteObjectTraits({.bWriteAttachments = 1U, .OOBWriteAttachmentType = (uint32)ENetObjectAttachmentType::OutOfBand}));
 			if (Result.Status == EWriteStatus::Abort)
 			{
 				return WrittenObjectCount;
@@ -3419,12 +3569,33 @@ uint32 FReplicationWriter::WriteOOBAttachments(FNetSerializationContext& Context
 			WrittenObjectCount += Result.NumWritten;
 		}
 	}
+	else if (WriteContext.WriteMode == EDataStreamWriteMode::DebugData)
+	{
+		if (WriteContext.bHasDebugObjectToSend && (ReplicationRecord.GetUnusedInfoCount() > 0U))
+		{
+			IRIS_PROFILER_SCOPE(FReplicationWriter_WriteDebugObjectAttachments);
+			const FWriteBatchResult Result = WriteObjectBatch(Context, ObjectIndexForOOBAttachment, FWriteObjectTraits({.bWriteAttachments = 1U, .OOBWriteAttachmentType = (uint32)ENetObjectAttachmentType::DebugObject}));
+			if (Result.Status == EWriteStatus::Abort)
+			{
+				return WrittenObjectCount;
+			}
+
+			const bool bHasDebugObjectToSend = Attachments.HasUnsentAttachments(ENetObjectAttachmentType::DebugObject, ObjectIndexForOOBAttachment);
+			WriteContext.bHasDebugObjectToSend = bHasDebugObjectToSend;
+			if (!bHasDebugObjectToSend)
+			{
+				DebugObjectSendQueue.Stats.EndSendingTime = FPlatformTime::Cycles64();
+			}
+
+			WrittenObjectCount += Result.NumWritten;
+		}
+	}
 	else
 	{
 		if (WriteContext.bHasHugeObjectToSend && (ReplicationRecord.GetUnusedInfoCount() > 0U))
 		{
 			IRIS_PROFILER_SCOPE(FReplicationWriter_WriteHugeObjectAttachments);
-			const FWriteBatchResult Result = WriteObjectBatch(Context, ObjectIndexForOOBAttachment, WriteObjectFlag_Attachments | WriteObjectFlag_HugeObject);
+			const FWriteBatchResult Result = WriteObjectBatch(Context, ObjectIndexForOOBAttachment, FWriteObjectTraits({.bWriteAttachments = 1U, .OOBWriteAttachmentType = (uint32)ENetObjectAttachmentType::HugeObject}));
 			if (Result.Status == EWriteStatus::Abort)
 			{
 				return WrittenObjectCount;
@@ -3443,7 +3614,7 @@ uint32 FReplicationWriter::WriteOOBAttachments(FNetSerializationContext& Context
 		if (WriteContext.bHasOOBAttachmentsToSend && CanSendObject(ObjectIndexForOOBAttachment) && (ReplicationRecord.GetUnusedInfoCount() > 0U))
 		{
 			IRIS_PROFILER_SCOPE(FReplicationWriter_WriteOOBAttachments);
-			const FWriteBatchResult Result = WriteObjectBatch(Context, ObjectIndexForOOBAttachment, WriteObjectFlag_Attachments);
+			const FWriteBatchResult Result = WriteObjectBatch(Context, ObjectIndexForOOBAttachment, FWriteObjectTraits({.bWriteAttachments = 1U, .OOBWriteAttachmentType = (uint32)ENetObjectAttachmentType::OutOfBand}));
 			if (Result.Status == EWriteStatus::Abort)
 			{
 				return WrittenObjectCount;
@@ -3519,7 +3690,7 @@ uint32 FReplicationWriter::WriteObjects(FNetSerializationContext& Context)
 			return EWriteStatus::Written;
 		}
 
-		const FWriteBatchResult Result = this->WriteObjectBatch(Context, InternalIndex, WriteObjectFlag_State | WriteObjectFlag_Attachments);
+		const FWriteBatchResult Result = this->WriteObjectBatch(Context, InternalIndex, FWriteObjectTraits({.bWriteState = 1U, .bWriteAttachments = 1U}));
 
 		if (Result.Status == EWriteStatus::Written)
 		{
@@ -3529,14 +3700,8 @@ uint32 FReplicationWriter::WriteObjects(FNetSerializationContext& Context)
 		return Result.Status;
 	};
 	
-	while (bContinue && ObjectListIt < ObjectCount)
+	while (bContinue && (ObjectListIt < ObjectCount || !WriteContext.AdditionalObjectSchedulingStack.IsEmpty()))
 	{
-		// Partial sort next batch
-		if (ObjectListIt >= SortedCount)
-		{
-			SortedCount += SortScheduledObjects(ObjectList, ObjectCount, ObjectListIt);
-		}
-
 		for (;;)
 		{
 			// Process objects that has been explicitly scheduled due to dependencies etc.
@@ -3554,7 +3719,7 @@ uint32 FReplicationWriter::WriteObjects(FNetSerializationContext& Context)
 				const EWriteStatus Status = SendObjectFunction(ExtraScheduleInfo.InternalIndex, ExtraScheduleInfo.bAppendCreationDependencies);
 				if (Status == EWriteStatus::Abort)
 				{
-					//UE_LOG_REPWRITER(VeryVerbose, TEXT("Aborted writing of additionally scheduled object: %s. Added back to the stack."), *NetRefHandleManager->PrintObjectFromIndex(ExtraScheduleInfo.InternalIndex));
+					UE_LOG_REPWRITER(VeryVerbose, TEXT("Aborted writing of additionally scheduled object: %s. Added back to the stack."), *NetRefHandleManager->PrintObjectFromIndex(ExtraScheduleInfo.InternalIndex));
 
 					// If we fail, we put the object back on the pending send stack and try again in the next packet of the batch
 					// The reason for pop before using the index is that the SendObjectFunction will push new dependent objects on the stack
@@ -3567,10 +3732,20 @@ uint32 FReplicationWriter::WriteObjects(FNetSerializationContext& Context)
 				if (Status == EWriteStatus::Skipped)
 				{
 					//UE_LOG_REPWRITER(VeryVerbose, TEXT("Skipping over additional scheduled object: %s"), *NetRefHandleManager->PrintObjectFromIndex(ExtraScheduleInfo.InternalIndex));
+				}
+				else if (Status == EWriteStatus::SkippedRetryNextPacket)
+				{
+					//UE_LOG_REPWRITER(VeryVerbose, TEXT("Skipping over additional scheduled object: %s will retry next packet"), *NetRefHandleManager->PrintObjectFromIndex(ExtraScheduleInfo.InternalIndex));
 					WriteContext.ObjectsSkippedInThisPacket.AddUnique(ExtraScheduleInfo.InternalIndex);
 				}
 			}
 
+			// Partial sort next batch
+			if (ObjectListIt >= SortedCount)
+			{
+				SortedCount += SortScheduledObjects(ObjectList, ObjectCount, ObjectListIt);
+			}
+
 			// Pick next object from original scheduling list.
 			if (bContinue && ObjectListIt < SortedCount)
 			{
@@ -3591,9 +3766,11 @@ uint32 FReplicationWriter::WriteObjects(FNetSerializationContext& Context)
 
 				if (Status == EWriteStatus::Skipped)
 				{
-					// Continue trying to send objects as we might find an object with a smaller footprint that fits in the available buffer space.
-					// Note that skipped objects won't be reserialized again even if the more packets are available after this batch.
 					//UE_LOG_REPWRITER(VeryVerbose, TEXT("Skipped writing of regular object: %s"), *NetRefHandleManager->PrintObjectFromIndex(InternalIndex));
+				}
+				else if (Status == EWriteStatus::SkippedRetryNextPacket)
+				{
+					//UE_LOG_REPWRITER(VeryVerbose, TEXT("Skipped writing of regular object will retry next packet if allowed: %s"), *NetRefHandleManager->PrintObjectFromIndex(InternalIndex));
 					WriteContext.ObjectsSkippedInThisPacket.AddUnique(InternalIndex);
 				}
 
@@ -3871,6 +4048,7 @@ UDataStream::EWriteResult FReplicationWriter::BeginWrite(const UDataStream::FBeg
 	WriteContext.bHasDestroyedObjectsToSend = 0U;
 	WriteContext.bHasHugeObjectToSend = 0U;
 	WriteContext.bHasOOBAttachmentsToSend = 0U;
+	WriteContext.bHasDebugObjectToSend = 0U;
 	WriteContext.ScheduledObjectCount = 0u;
 
 	WriteContext.WriteMode = Params.WriteMode;
@@ -3885,6 +4063,15 @@ UDataStream::EWriteResult FReplicationWriter::BeginWrite(const UDataStream::FBeg
 		}
 		WriteContext.bHasOOBAttachmentsToSend = bHasUnsentOOBAttachments;
 	}
+	else if (WriteContext.WriteMode == EDataStreamWriteMode::DebugData)
+	{
+		const bool bHasUnsentDebugObject = Attachments.HasUnsentAttachments(ENetObjectAttachmentType::DebugObject, ObjectIndexForOOBAttachment);
+		if (!bHasUnsentDebugObject)
+		{
+			return UDataStream::EWriteResult::NoData;
+		}
+		WriteContext.bHasDebugObjectToSend = bHasUnsentDebugObject;
+	}
 	else
 	{
 		// See if we have any work to do
@@ -4019,7 +4206,7 @@ void FReplicationWriter::Update(const UDataStream::FUpdateParameters& Params)
 
 bool FReplicationWriter::HasDataToSend(const FWriteContext& Context) const
 {
-	return WriteContext.bIsValid & (Context.bHasDestroyedObjectsToSend | Context.bHasUpdatedObjectsToSend | Context.bHasHugeObjectToSend | Context.bHasOOBAttachmentsToSend);
+	return WriteContext.bIsValid & (Context.bHasDestroyedObjectsToSend | Context.bHasUpdatedObjectsToSend | Context.bHasHugeObjectToSend | Context.bHasOOBAttachmentsToSend | Context.bHasDebugObjectToSend);
 }
 
 UDataStream::EWriteResult FReplicationWriter::Write(FNetSerializationContext& Context)
@@ -4038,7 +4225,7 @@ UDataStream::EWriteResult FReplicationWriter::Write(FNetSerializationContext& Co
 	WriteContext.bIsInReplicationRecordStarvation = ReplicationRecord.GetUnusedInfoCount() < static_cast<uint32>(GReplicationWriterReplicationRecordStarvationThreshold);
 	
 	// Packets going over the bandwidth limit are considered an OOB packet which should avoid writing things like objects pending destroy.  PostTickDispath packets are also flagged OOB to minimize extra serialization data.
-	WriteContext.bIsOOBPacket = (WriteContext.WriteMode == EDataStreamWriteMode::PostTickDispatch) || (WriteContext.MaxPacketsToSend != 0 && WriteContext.NumWrittenPacketsInThisBatch >= WriteContext.MaxPacketsToSend);
+	WriteContext.bIsOOBPacket = (WriteContext.WriteMode == EDataStreamWriteMode::PostTickDispatch) || (WriteContext.WriteMode == EDataStreamWriteMode::DebugData) || (WriteContext.MaxPacketsToSend != 0 && WriteContext.NumWrittenPacketsInThisBatch >= WriteContext.MaxPacketsToSend);
 
 	FNetBitStreamWriter& Writer = *Context.GetBitStreamWriter();
 
@@ -4091,7 +4278,7 @@ UDataStream::EWriteResult FReplicationWriter::Write(FNetSerializationContext& Co
 	// https://jira.it.epicgames.com/browse/UE-127371
 	if (HasDataToSend(WriteContext))
 	{
-		if (WriteContext.bHasHugeObjectToSend || (WriteContext.bHasOOBAttachmentsToSend && CanSendObject(ObjectIndexForOOBAttachment)))
+		if (WriteContext.bHasHugeObjectToSend || WriteContext.bHasDebugObjectToSend || (WriteContext.bHasOOBAttachmentsToSend && CanSendObject(ObjectIndexForOOBAttachment)))
 		{
 			WriteResult = UDataStream::EWriteResult::HasMoreData;
 		}
@@ -4130,6 +4317,14 @@ UDataStream::EWriteResult FReplicationWriter::Write(FNetSerializationContext& Co
 			FNetTrace::FoldTraceCollector(Context.GetTraceCollector(), Collector, GetBitStreamPositionForNetTrace(Writer));
 			Collector->Reset();
 		}
+		if (WriteContext.WriteMode == EDataStreamWriteMode::DebugData)
+		{
+			if (FNetTraceCollector* Collector = DebugObjectSendQueue.TraceCollector)
+			{
+				FNetTrace::FoldTraceCollector(Context.GetTraceCollector(), Collector, GetBitStreamPositionForNetTrace(Writer));
+				Collector->Reset();
+			}
+		}
 #endif
 
 		++WriteContext.NumWrittenPacketsInThisBatch;
@@ -4268,9 +4463,26 @@ bool FReplicationWriter::CanQueueHugeObject() const
 	return true;
 }
 
-void FReplicationWriter::FreeHugeObjectSendQueue()
+bool FReplicationWriter::CanQueueDebugObject() const
+{
+	if (DebugObjectSendQueue.IsFull())
+	{
+		return false;
+	}
+
+	// Check whether the reliable queue is full in which case there's no point in queueing additional huge objects.
+	if (!Attachments.CanSendMoreReliableAttachments(ENetObjectAttachmentType::DebugObject, ObjectIndexForOOBAttachment))
+	{
+		return false;
+	}
+
+	return true;
+}
+
+
+void FReplicationWriter::FreeHugeObjectSendQueues()
 {
-	HugeObjectSendQueue.FreeContexts([this](const FHugeObjectContext& HugeObjectContext)
+	auto FreeContextFunc = [this](const FHugeObjectContext& HugeObjectContext)
 	{
 		for (const FObjectRecord& ObjectRecord : HugeObjectContext.BatchRecord.ObjectReplicationRecords)
 		{
@@ -4281,7 +4493,9 @@ void FReplicationWriter::FreeHugeObjectSendQueue()
 				FChangeMaskStorageOrPointer::Free(ObjectRecord.Record.ChangeMaskOrPtr, ChangeMaskBitCount, s_DefaultChangeMaskAllocator);
 			}
 		}
-	});
+	};
+	HugeObjectSendQueue.FreeContexts(FreeContextFunc);
+	DebugObjectSendQueue.FreeContexts(FreeContextFunc);
 }
 
 void FReplicationWriter::CollectAndAppendExports(FNetSerializationContext& Context, uint8* RESTRICT InternalBuffer, const FReplicationProtocol* Protocol) const
@@ -4435,8 +4649,9 @@ FReplicationWriter::FHugeObjectContext::FHugeObjectContext() = default;
 FReplicationWriter::FHugeObjectContext::~FHugeObjectContext() = default;
 
 // HugeObjectSendQueue implementation
-FReplicationWriter::FHugeObjectSendQueue::FHugeObjectSendQueue()
-: DebugName(CreatePersistentNetDebugName(TEXT("HugeObjectState"), UE_ARRAY_COUNT(TEXT("HugeObjectState"))))
+FReplicationWriter::FHugeObjectSendQueue::FHugeObjectSendQueue(const TCHAR* QueueName, bool bInAllowMultipleBatchesPerRootObjectInTransit)
+: DebugName(CreatePersistentNetDebugName(QueueName))
+, bAllowMultipleBatchesPerRootObjectInTransit(bInAllowMultipleBatchesPerRootObjectInTransit)
 {
 #if UE_NET_TRACE_ENABLED
 	DebugName->DebugNameId = FNetTrace::TraceName(DebugName->Name);
@@ -4473,13 +4688,25 @@ bool FReplicationWriter::FHugeObjectSendQueue::EnqueueHugeObject(const FHugeObje
 		return false;
 	}
 
-	if (RootObjectsInTransit.Find(Context.RootObjectInternalIndex))
+	int32* RootObjectInTransitEntry = RootObjectsInTransit.Find(Context.RootObjectInternalIndex);
+	if (RootObjectInTransitEntry)
 	{
-		ensureMsgf(false, TEXT("An object that is already in the huge object queue should not try replicating again ( InternalIndex: %u )"), Context.RootObjectInternalIndex);
-		return false;
+		if (!bAllowMultipleBatchesPerRootObjectInTransit)
+		{
+			ensureMsgf(false, TEXT("An object that is already in the huge object queue should not try replicating again ( InternalIndex: %u )"), Context.RootObjectInternalIndex);
+			return false;
+		}
+		else
+		{
+			// Bump the number of contexts in flight for this rootobject.
+			++*RootObjectInTransitEntry;
+		}		
+	}
+	else
+	{
+		RootObjectsInTransit.Add(Context.RootObjectInternalIndex, 1);
 	}
 
-	RootObjectsInTransit.Add(Context.RootObjectInternalIndex);
 	// Note: Lists don't have methods to perform moving of an element.
 	SendContexts.AddTail(Context);
 	return true;
@@ -4493,7 +4720,7 @@ bool FReplicationWriter::FHugeObjectSendQueue::IsObjectInQueue(FInternalNetRefIn
 		return false;
 	}
 
-	if (RootObjectsInTransit.Find(ObjectIndex))
+	if (RootObjectsInTransit.Contains(ObjectIndex))
 	{
 		return true;
 	}
@@ -4563,12 +4790,20 @@ void FReplicationWriter::FHugeObjectSendQueue::AckObjects(TFunctionRef<void (con
 		AckHugeObject(Context);
 
 		// Remove from fast lookup set.
-		RootObjectsInTransit.Remove(Context.RootObjectInternalIndex);
+		int32* RootObjectsInTransitEntry = RootObjectsInTransit.Find(Context.RootObjectInternalIndex);
+		if (RootObjectsInTransitEntry)
+		{
+			// If this was the last entry we remove it
+			if (--(*RootObjectsInTransitEntry) == 0)
+			{
+				RootObjectsInTransit.Remove(Context.RootObjectInternalIndex);
+			}
+		}
 
 		// Remove from queue.
 		SendContexts.RemoveNode(Node);
 
-		ensure(SendContexts.IsEmpty() == RootObjectsInTransit.IsEmpty());
+		ensure((SendContexts.IsEmpty() == RootObjectsInTransit.IsEmpty()));
 	}
 
 	if (RootObjectsInTransit.IsEmpty())
diff --git a/Engine/Source/Runtime/Net/Iris/Private/Iris/ReplicationSystem/ReplicationWriter.h b/Engine/Source/Runtime/Net/Iris/Private/Iris/ReplicationSystem/ReplicationWriter.h
--- a/Engine/Source/Runtime/Net/Iris/Private/Iris/ReplicationSystem/ReplicationWriter.h
+++ b/Engine/Source/Runtime/Net/Iris/Private/Iris/ReplicationSystem/ReplicationWriter.h
@@ -73,6 +73,7 @@ public:
 		// Special states
 		AttachmentToObjectNotInScope,	// Special state for object index used for sending attachments to remote owned objects
 		HugeObject,						// Special state for object index used for sending parts of huge object payloads
+		DebugObject,					// Special state for object index used for sending parts of debug object payloads
 
 		// Normal states
 		PendingCreate,					// Not yet created, no data in flight
@@ -134,6 +135,7 @@ public:
 				uint64 FlushFlags : 3;									// Flags indicating what we are waiting for when flushing
 				uint64 HasDirtyConditionals : 1;						// If this flag is set, we must update conditionals.
 				mutable uint64 HasCannotSendInfo : 1;					// If this flag is set, the object has been prevented from being sent at least once.
+				uint64 IsDebugObject : 1;								// If this flag is set, the object is a debugobject and is always sent using DebugObject attachments.
 			};
 		};
 
@@ -154,6 +156,7 @@ public:
 	};
 
 public:
+	FReplicationWriter();
 	~FReplicationWriter();
 
 	// Init
@@ -203,6 +206,8 @@ public:
 
 	bool AreAllReliableAttachmentsSentAndAcked() const;
 
+	bool HasDebugDataToSend() const;
+
 	void Update(const UDataStream::FUpdateParameters& Params);
 
 	[[nodiscard]] FString PrintObjectInfo(FInternalNetRefIndex ObjectIndex) const;
@@ -295,6 +300,9 @@ private:
 		uint32 bHasHugeObjectToSend : 1;
 		// Whether there are OOB attachments to send or not
 		uint32 bHasOOBAttachmentsToSend : 1;
+		// Whether there is at least one debug object to send or not
+		uint32 bHasDebugObjectToSend : 1;
+
 		// Whether this packet is mainly intended for OOB and HugeObject data
 		uint32 bIsOOBPacket : 1;
 		// Whether this context is valid or not
@@ -324,6 +332,7 @@ private:
 	{
 		Object,
 		HugeObject,
+		DebugObject,
 		OOBAttachment,
 		Internal,
 		// Currently unused as destruction infos don't create a BatchInfo.
@@ -378,7 +387,7 @@ private:
 	class FHugeObjectSendQueue
 	{
 	public:
-		FHugeObjectSendQueue();
+		FHugeObjectSendQueue(const TCHAR* QueueName, bool bInAllowMultipleBatchesPerRootObjectInTransit = false);
 		~FHugeObjectSendQueue();
 
 		// If the queue is full we can't start another send.
@@ -418,16 +427,17 @@ private:
 		const FNetDebugName* DebugName = nullptr;
 
 	private:
-		TSet<FInternalNetRefIndex> RootObjectsInTransit;
+		TMap<FInternalNetRefIndex, int32> RootObjectsInTransit;
 		TDoubleLinkedList<FHugeObjectContext> SendContexts;
+		bool bAllowMultipleBatchesPerRootObjectInTransit = false;
 	};
 
-	enum EWriteObjectFlag : unsigned
+	struct FWriteObjectTraits
 	{
-		WriteObjectFlag_State = 1U,
-		WriteObjectFlag_Attachments = WriteObjectFlag_State << 1U,
-		WriteObjectFlag_HugeObject = WriteObjectFlag_Attachments << 1U,
-		WriteObjectFlag_IsWritingHugeObjectBatch = WriteObjectFlag_HugeObject << 1U,
+		uint32 bWriteState : 1U = 0U;
+		uint32 bWriteAttachments : 1U = 0U;
+		uint32 OOBWriteAttachmentType : 2U = 0U;
+		uint32 bIsWritingHugeObjectBatch : 1U = 0U;
 	};
 
 	enum class EWriteObjectRetryMode : unsigned
@@ -520,10 +530,10 @@ private:
 	uint32 WriteRootObjectsPendingDestroy(FNetSerializationContext& Context);
 
 	// Write object and SubObjects
-	EWriteObjectStatus WriteObjectAndSubObjects(FNetSerializationContext& Context, uint32 InternalIndex, uint32 WriteObjectFlags, FBatchInfo& OutBatchInfo);
+	EWriteObjectStatus WriteObjectAndSubObjects(FNetSerializationContext& Context, uint32 InternalIndex, FWriteObjectTraits WriteObjectTraits, FBatchInfo& OutBatchInfo);
 
 	// Write objects recursive
-	EWriteObjectStatus WriteObjectInBatch(FNetSerializationContext& Context, uint32 InternalIndex, uint32 WriteObjectFlags, FBatchInfo& OutBatchInfo);
+	EWriteObjectStatus WriteObjectInBatch(FNetSerializationContext& Context, uint32 InternalIndex, FWriteObjectTraits WriteObjectTraits, FBatchInfo& OutBatchInfo);
 
 	enum class EWriteStatus : int32
 	{
@@ -531,10 +541,12 @@ private:
 		Abort = -1,
 		// Object was skipped and won't be retried again. Ex: failed dependency, waiting for creation confirmation or buffer won't fit the object and its the last packet to send.
 		Skipped = 0,
+		// Object was skipped due to not enough space but will be retried in next packet if we are allowed to write more data.
+		SkippedRetryNextPacket = 1,
 		// Object was successfully written
-		Written = 1,
+		Written = 2,
 		// Object is rescheduled due to creation dependency
-		RescheduledDueToCreationDependency = 2,
+		RescheduledDueToCreationDependency = 3,
 	};
 
 	// Write destruction info for an object that should be destroyed
@@ -551,9 +563,9 @@ private:
 
 	// Write Object and any subobject(s) to stream as an atomic batch
 	// returns > 1 if Objects was written, 0 if the objects was skipped (failed dependency or waiting for creation confirmation) -1 if the stream is full and we should stop
-	FWriteBatchResult WriteObjectBatch(FNetSerializationContext& Context, FInternalNetRefIndex InternalIndex, uint32 WriteObjectFlags);
+	FWriteBatchResult WriteObjectBatch(FNetSerializationContext& Context, FInternalNetRefIndex InternalIndex, FWriteObjectTraits WriteObjectTraits);
 
-	EWriteStatus PrepareAndSendHugeObjectPayload(FNetSerializationContext& Context, FInternalNetRefIndex InternalIndex);
+	EWriteStatus PrepareAndSendHugeObjectPayload(FNetSerializationContext& Context, FInternalNetRefIndex InternalIndex, ENetObjectAttachmentType AttachmentType);
 
 	// Write OOBAttachments
 	uint32 WriteOOBAttachments(FNetSerializationContext& Context);
@@ -605,8 +617,9 @@ private:
 	bool IsObjectPartOfActiveHugeObject(uint32 InternalIndex) const;
 
 	bool CanQueueHugeObject() const;
+	bool CanQueueDebugObject() const;
 
-	void FreeHugeObjectSendQueue();
+	void FreeHugeObjectSendQueues();
 
 	bool HasDataToSend(const FWriteContext& Context) const;
 
@@ -671,6 +684,8 @@ private:
 	FWriteContext WriteContext;
 	FBitStreamInfo WriteBitStreamInfo;
 	FHugeObjectSendQueue HugeObjectSendQueue;
+	FHugeObjectSendQueue DebugObjectSendQueue;
+
 	// Features that can aid in tracking down bitstream errors etc.
 	EReplicationDataStreamDebugFeatures StreamDebugFeatures = EReplicationDataStreamDebugFeatures::None;
 
