Iris
* Added NetRefHandle's to UNetObjectFactory::DestroyedContext struct
* Misc code and log improvements

#rnx

diff --git a/Engine/Source/Runtime/Net/Iris/Private/Iris/ReplicationSystem/ObjectReplicationBridge.cpp b/Engine/Source/Runtime/Net/Iris/Private/Iris/ReplicationSystem/ObjectReplicationBridge.cpp
--- a/Engine/Source/Runtime/Net/Iris/Private/Iris/ReplicationSystem/ObjectReplicationBridge.cpp
+++ b/Engine/Source/Runtime/Net/Iris/Private/Iris/ReplicationSystem/ObjectReplicationBridge.cpp
@@ -1198,13 +1198,14 @@ void UObjectReplicationBridge::DetachInstanceFromRemote(FNetRefHandle Handle, ER
 	UNetObjectFactory::FDestroyedContext DestroyContext
 	{
 		.DestroyedInstance = GetObjectFromReferenceHandle(Handle),
+		.DestroyedObjectHandle = Handle,
 		.DestroyReason = DestroyReason,
 		.DestroyFlags = DestroyFlags,
 	};
 
 	if (DestroyContext.DestroyedInstance)
 	{
-		// Check if it's a subobject and we need to tell it's RootObject
+		// Check if it's a subobject
 		{
 			const FNetRefHandleManager::FReplicatedObjectData& DestroyedObjectData = NetRefHandleManager->GetReplicatedObjectData(NetRefHandleManager->GetInternalIndex(Handle));
 
@@ -1213,7 +1214,9 @@ void UObjectReplicationBridge::DetachInstanceFromRemote(FNetRefHandle Handle, ER
 			if (RootObjectIndex != FNetRefHandleManager::InvalidInternalIndex)
 			{
 				DestroyContext.RootObject = NetRefHandleManager->GetReplicatedObjectInstance(RootObjectIndex);
+				DestroyContext.RootObjectHandle = NetRefHandleManager->GetNetRefHandleFromInternalIndex(RootObjectIndex);
 
+				// Tell the RootObject one of his subobjects was detached
 				if (UNetObjectFactory* RootObjectFactory = GetNetFactory(NetRefHandleManager->GetReplicatedObjectData(RootObjectIndex).NetFactoryId))
 				{
 					RootObjectFactory->SubObjectDetachedFromReplication(DestroyContext);
@@ -1221,21 +1224,13 @@ void UObjectReplicationBridge::DetachInstanceFromRemote(FNetRefHandle Handle, ER
 			}
 		}
 
-		// Now ask the object's factory to handle the destruction
+		// Now ask the object's factory to handle the detachment
 		if (ensureMsgf(NetFactoryId != InvalidNetObjectFactoryId, TEXT("Received invalid netfactory for %s: %s"), *Handle.ToString(), *GetPathNameSafe(DestroyContext.DestroyedInstance)))
 		{
 			if (DestroyReason != EReplicationBridgeDestroyInstanceReason::Migrated)
 			{
 				if (UNetObjectFactory* Factory = GetNetFactory(NetFactoryId))
 				{
-PRAGMA_DISABLE_DEPRECATION_WARNINGS
-					// Call the deprecated function
-					if (DestroyReason != EReplicationBridgeDestroyInstanceReason::DoNotDestroy)
-					{
-						Factory->DestroyReplicatedObject(DestroyContext);
-					}
-PRAGMA_ENABLE_DEPRECATION_WARNINGS
-
 					Factory->DetachedFromReplication(DestroyContext);
 				}
 			}
diff --git a/Engine/Source/Runtime/Net/Iris/Private/Iris/ReplicationSystem/ReplicationBridge.cpp b/Engine/Source/Runtime/Net/Iris/Private/Iris/ReplicationSystem/ReplicationBridge.cpp
--- a/Engine/Source/Runtime/Net/Iris/Private/Iris/ReplicationSystem/ReplicationBridge.cpp
+++ b/Engine/Source/Runtime/Net/Iris/Private/Iris/ReplicationSystem/ReplicationBridge.cpp
@@ -234,7 +234,7 @@ void UReplicationBridge::InternalAttachInstanceToNetRefHandle(FNetRefHandle RefH
 	const FInternalNetRefIndex InternalReplicationIndex = NetRefHandleManager->GetInternalIndex(RefHandle);
 
 	NetRefHandleManager->AttachInstanceProtocol(InternalReplicationIndex, InstanceProtocol, Instance);
-	UE_LOG_REPLICATIONBRIDGE(Verbose, TEXT("InternalAttachInstanceToNetHandle Attached: %s %s to (InternalIndex: %u)"), *Instance->GetName(), *RefHandle.ToString(), InternalReplicationIndex);
+	UE_LOG_REPLICATIONBRIDGE(Verbose, TEXT("InternalAttachInstanceToNetHandle Attached: %s (0x%p) %s to (InternalIndex: %u)"), *Instance->GetName(), Instance, *RefHandle.ToString(), InternalReplicationIndex);
 
 	// Bind instance protocol to dirty state tracking
 	if (bBindInstanceProtocol)
diff --git a/Engine/Source/Runtime/Net/Iris/Private/Iris/ReplicationSystem/ReplicationReader.cpp b/Engine/Source/Runtime/Net/Iris/Private/Iris/ReplicationSystem/ReplicationReader.cpp
--- a/Engine/Source/Runtime/Net/Iris/Private/Iris/ReplicationSystem/ReplicationReader.cpp
+++ b/Engine/Source/Runtime/Net/Iris/Private/Iris/ReplicationSystem/ReplicationReader.cpp
@@ -219,16 +219,15 @@ struct FReplicationReader::FDispatchObjectInfo
 {
 	FInternalNetRefIndex InternalIndex = FNetRefHandleManager::InvalidInternalIndex;
 	FChangeMaskStorageOrPointer ChangeMaskOrPointer;
-	uint32 bIsInitialState : 1 = false;
-	uint32 bHasState : 1 = false;
-	uint32 bHasAttachments : 1 = false;
-	uint32 bDestroy : 1 = false;
-	uint32 bTearOff : 1 = false;
-	uint32 bDeferredEndReplication : 1 = false;
-	uint32 bReplicatedDestroyForSubObject : 1 = false;
-	uint32 bProcessEndReplicationForHierarchy : 1 = false;
-	uint32 bShouldCallSubObjectCreatedFromReplication : 1 = false;
-	uint32 bDynamicObjectCreated : 1 = false;
+	bool bIsInitialState : 1 = false;
+	bool bHasState : 1 = false;
+	bool bHasAttachments : 1 = false;
+	bool bDestroy : 1 = false;
+	bool bTearOff : 1 = false;
+	bool bDeferredEndReplication : 1 = false;
+	bool bReplicatedDestroyForSubObject : 1 = false;
+	bool bProcessEndReplicationForHierarchy : 1 = false;
+	bool bShouldCallSubObjectCreatedFromReplication : 1 = false;
 	ECloseReplicationReason CloseReason = ECloseReplicationReason::None;
 };
 
@@ -540,10 +539,10 @@ uint32 FReplicationReader::ReadObjectsAndSubObjectsPendingDestroy(FNetSerializat
 					Info.bTearOff = false;
 					Info.bDeferredEndReplication = true;
 					Info.InternalIndex = InternalIndex;
-					Info.bIsInitialState = 0U;
-					Info.bHasState = 0U;
-					Info.bHasAttachments = 0U;
-					Info.bShouldCallSubObjectCreatedFromReplication = 0U;
+					Info.bIsInitialState = false;
+					Info.bHasState = false;
+					Info.bHasAttachments = false;
+					Info.bShouldCallSubObjectCreatedFromReplication = false;
 					Info.CloseReason = CloseReason;
 
 					// Mark for dispatch
@@ -622,10 +621,10 @@ uint32 FReplicationReader::ReadRootObjectsPendingDestroy(FNetSerializationContex
 					Info.bDeferredEndReplication = true;
 					Info.bProcessEndReplicationForHierarchy = true;
 					Info.InternalIndex = InternalIndex;
-					Info.bIsInitialState = 0U;
-					Info.bHasState = 0U;
-					Info.bHasAttachments = 0U;
-					Info.bShouldCallSubObjectCreatedFromReplication = 0U;
+					Info.bIsInitialState = false;
+					Info.bHasState = false;
+					Info.bHasAttachments = false;
+					Info.bShouldCallSubObjectCreatedFromReplication = false;
 					Info.CloseReason = CloseReason;
 
 					// Mark for dispatch
@@ -717,7 +716,6 @@ void FReplicationReader::EndReplication(FInternalNetRefIndex InternalIndex, bool
 	
 	if (FReplicatedObjectInfo* ObjectInfo = ReplicatedObjects.Find(InternalIndex))
 	{
-
 		CleanupReferenceTracking(ObjectInfo);
 		Attachments.DropAllAttachments(ENetObjectAttachmentType::Normal, InternalIndex);
 
@@ -1467,12 +1465,12 @@ void FReplicationReader::ReadObjectInBatch(FNetSerializationContext& Context, FN
 		// We set one of these flags
 		if (bImmediateDispatchEndReplicationForSubObjects && Info.bDestroy && !Info.bTearOff)
 		{
-			Info.bReplicatedDestroyForSubObject = 1U;
-			Info.bDeferredEndReplication = 0U;
+			Info.bReplicatedDestroyForSubObject = true;
+			Info.bDeferredEndReplication = false;
 		}
 		else
 		{
-			Info.bReplicatedDestroyForSubObject = 0U;
+			Info.bReplicatedDestroyForSubObject = false;
 			Info.bDeferredEndReplication = !!(ReplicatedDestroyHeaderFlags & (ReplicatedDestroyHeaderFlags_TearOff | ReplicatedDestroyHeaderFlags_EndReplication));
 		}
 		Info.bShouldCallSubObjectCreatedFromReplication = bShouldCallSubObjectCreatedFromReplication;
@@ -1593,9 +1591,9 @@ void FReplicationReader::ReadObjectInBatch(FNetSerializationContext& Context, FN
 		if (bShouldCommitPendingDispatchObjectInfo)
 		{
 			Info.InternalIndex = InternalIndex;
-			Info.bIsInitialState = bIsInitialState ? 1U : 0U;
-			Info.bHasState = bHasState ? 1U : 0U;
-			Info.bHasAttachments = bHasAttachments ? 1U : 0U;
+			Info.bIsInitialState = bIsInitialState;
+			Info.bHasState = bHasState;
+			Info.bHasAttachments = bHasAttachments;
 
 			ObjectsToDispatchArray->CommitPendingDispatchObjectInfo();
 		}
@@ -2378,7 +2376,7 @@ void FReplicationReader::DispatchStateData(FNetSerializationContext& Context)
 
 			// Collect all unresolvable references, including old pending references
 			FResolveAndCollectUnresolvedAndResolvedReferenceCollector Collector;
-			Collector.CollectReferences(*ObjectReferenceCache, ResolveContext, Info.bIsInitialState | ReplicationInfo->bHasUnresolvedInitialReferences, &ChangeMaskForResolve, ObjectData.ReceiveStateBuffer, ObjectData.Protocol);
+			Collector.CollectReferences(*ObjectReferenceCache, ResolveContext, Info.bIsInitialState || (ReplicationInfo->bHasUnresolvedInitialReferences!=0U), &ChangeMaskForResolve, ObjectData.ReceiveStateBuffer, ObjectData.Protocol);
 
 			// If we have or had any object references we need to track them and update the unresolved mask
 			if (bHadUnresolvedReferences || Collector.GetUnresolvedReferences().Num() > 0 || Collector.GetResolvedReferences().Num() > 0)
diff --git a/Engine/Source/Runtime/Net/Iris/Public/Iris/ReplicationSystem/NetObjectFactory.h b/Engine/Source/Runtime/Net/Iris/Public/Iris/ReplicationSystem/NetObjectFactory.h
--- a/Engine/Source/Runtime/Net/Iris/Public/Iris/ReplicationSystem/NetObjectFactory.h
+++ b/Engine/Source/Runtime/Net/Iris/Public/Iris/ReplicationSystem/NetObjectFactory.h
@@ -147,12 +147,6 @@ public:
 	 */
 	virtual void DetachedFromReplication(const FDestroyedContext& Context) PURE_VIRTUAL(UNetObjectFactory::DetachedFromReplication, );
 
-	/**
-	 * Replaced by DetachedFromReplication since it didn't receive callbacks for objects with EReplicationBridgeDestroyInstanceReason::DoNotDestroy reasons.
-	*/
-	UE_DEPRECATED(5.7, "Replaced by DetachedFromReplication")
-	virtual void DestroyReplicatedObject(const FDestroyedContext& Context) final {}
-
 	/** 
 	 * Optional callback triggered when a root object managed by this factory gets assigned a dynamic subobject. 
 	 * This callback is called on remotes only. 
@@ -238,8 +232,16 @@ public:
 	{
 		/** The object about to be destroyed */
 		UObject* DestroyedInstance = nullptr;
+
 		/** Optional pointer to the root object when the destroyed object is a subobject */
 		UObject* RootObject = nullptr;
+
+		/** The handle of the object destroyed */
+		UE::Net::FNetRefHandle DestroyedObjectHandle;
+
+		/** Optional handle of the root object when the destroyed object is a subobject */
+		UE::Net::FNetRefHandle RootObjectHandle;
+
 		EReplicationBridgeDestroyInstanceReason DestroyReason = EReplicationBridgeDestroyInstanceReason::DoNotDestroy;
 		EReplicationBridgeDestroyInstanceFlags DestroyFlags = EReplicationBridgeDestroyInstanceFlags::None;
 	};
