Iris: adding FSharedStringNetToken and using it to store the RelativePath token in FCachedNetObjectReference. This allows the shared object reference cache to store these tokens and safely resolve or compare them despite the reference cache being used by multiple client replication systems.

diff --git a/Engine/Source/Runtime/Net/Iris/Private/Iris/ReplicationSystem/NetTokenStore.cpp b/Engine/Source/Runtime/Net/Iris/Private/Iris/ReplicationSystem/NetTokenStore.cpp
--- a/Engine/Source/Runtime/Net/Iris/Private/Iris/ReplicationSystem/NetTokenStore.cpp
+++ b/Engine/Source/Runtime/Net/Iris/Private/Iris/ReplicationSystem/NetTokenStore.cpp
@@ -227,7 +227,7 @@ FNetToken FNetTokenDataStore::ReadNetToken(UE::Net::FNetSerializationContext& Co
 	return FNetTokenStore::InternalReadNetToken(Context, GetTypeId());
 }
 
-void FNetTokenDataStore::WriteNetToken(UE::Net::FNetSerializationContext& Context, FNetToken Token)
+void FNetTokenDataStore::WriteNetToken(UE::Net::FNetSerializationContext& Context, FNetToken Token) const
 {
 	const bool bWriteTokenTypeId = false;
 	return FNetTokenStore::InternalWriteNetToken(Context, Token, bWriteTokenTypeId);
diff --git a/Engine/Source/Runtime/Net/Iris/Private/Iris/ReplicationSystem/ObjectReferenceCache.cpp b/Engine/Source/Runtime/Net/Iris/Private/Iris/ReplicationSystem/ObjectReferenceCache.cpp
--- a/Engine/Source/Runtime/Net/Iris/Private/Iris/ReplicationSystem/ObjectReferenceCache.cpp
+++ b/Engine/Source/Runtime/Net/Iris/Private/Iris/ReplicationSystem/ObjectReferenceCache.cpp
@@ -482,7 +482,7 @@ bool FObjectReferenceCache::CreateObjectReferenceInternal(const UObject* Object,
 	CachedObject.Object = MakeWeakObjectPtr(const_cast<UObject*>(Object));
 	CachedObject.ObjectKey = Object;
 	CachedObject.NetRefHandle = NetRefHandle;
-	CachedObject.RelativePath = PathToken;
+	CachedObject.RelativePath.AssignToken(PathToken, { .ReplicationSystemId = ReplicationSystem->GetId(), .ConnectionId = InvalidConnectionId });
 	CachedObject.OuterNetRefHandle = CreateObjectReferenceHandle(Object->GetOuter(), ReplicationSystem);
 	CachedObject.bIsPackage = bIsStatic && !Object->GetOuter();
 	CachedObject.bNoLoad = !CanClientLoadObjectInternal(Object, bIsDynamic); // We are probably going to deal with this as a dependency
@@ -610,7 +610,6 @@ void FObjectReferenceCache::BindRemoteReference(FNetRefHandle RefHandle, const U
 		CachedObject.Object = MakeWeakObjectPtr(const_cast<UObject*>(Object));
 		CachedObject.ObjectKey = Object;
 		CachedObject.NetRefHandle = RefHandle;
-		CachedObject.RelativePath = FNetToken();
 		CachedObject.OuterNetRefHandle = FNetRefHandle::GetInvalid();
 
 		UE_LOG_REFERENCECACHE(Verbose, TEXT("%hs: Adding new cache entry for: %s with Object: %s (0x%p)."), __FUNCTION__, *RefHandle.ToString(), *GetNameSafe(Object), Object);
@@ -923,8 +922,8 @@ UObject* FObjectReferenceCache::ResolveObjectReferenceHandleInternal(FNetRefHand
 		UE_LOG(LogIris, Error, TEXT("GetObjectFromRefHandle: Outer is null or package is not fully loaded. FullPath: %s Outer: %s"), ToCStr(FullPath(RefHandle, ResolveContext)), *GetFullNameSafe(ObjOuter));
 	}
 
-	FStringTokenStore* LocalStringTokenStore = ResolveContext.ReplicationSystem->GetNetTokenStore()->GetDataStore<FStringTokenStore>();
-	const TCHAR* ResolvedToken = LocalStringTokenStore->ResolveRemoteToken(CacheObjectPtr->RelativePath, *ResolveContext.RemoteNetTokenStoreState);
+	const TCHAR* ResolvedToken = CacheObjectPtr->RelativePath.ResolveToken();
+
 	if (ResolvedToken == nullptr)
 	{
 		return nullptr;
@@ -1349,7 +1348,7 @@ FNetObjectReference FObjectReferenceCache::GetOrCreateObjectReference(const FStr
 
 		CachedObject.ObjectKey = nullptr;
 		CachedObject.NetRefHandle = RefHandle;
-		CachedObject.RelativePath = LocalStringTokenStore->GetOrCreateToken(ObjectPath);
+		CachedObject.RelativePath.AssignToken(LocalStringTokenStore->GetOrCreateToken(ObjectPath), { .ReplicationSystemId = ReplicationSystem->GetId(), .ConnectionId = InvalidConnectionId });
 		CachedObject.OuterNetRefHandle = CreateObjectReferenceHandle(Outer, ReplicationSystem);
 		CachedObject.bIsPackage = 0U;
 		CachedObject.bNoLoad = 1U;
@@ -1405,12 +1404,9 @@ void FObjectReferenceCache::WriteFullReferenceInternal(FNetSerializationContext&
 
 		if (Writer->WriteBool(bHasPath))
 		{
-			FNetTokenStore* LocalNetTokenStore = Context.GetInternalContext()->ReplicationSystem->GetNetTokenStore();
-			FStringTokenStore* LocalStringTokenStore = Context.GetInternalContext()->ReplicationSystem->GetNetTokenStore()->GetDataStore<FStringTokenStore>();
-
 			// Note: As we use the StringTokenStore to write the token we also need to read it using the stringtokenstore.
-			LocalStringTokenStore->WriteNetToken(Context, CachedObject->RelativePath);
-			LocalNetTokenStore->ConditionalWriteNetTokenData(Context, ExportContext, CachedObject->RelativePath);
+			CachedObject->RelativePath.WriteNetToken(Context);
+			CachedObject->RelativePath.ConditionalWriteNetTokenData(Context, ExportContext);
 			WriteFullReferenceInternal(Context, FNetObjectReference(CachedObject->OuterNetRefHandle));
 		}
 	
@@ -1498,8 +1494,10 @@ void FObjectReferenceCache::ReadFullReferenceInternal(FNetSerializationContext&
 	// Already known? Validate that it is correct!
 	if (const FCachedNetObjectReference* CachedObject = ReferenceHandleToCachedReference.Find(NetRefHandle))
 	{
+		FNetTokenResolveContext ReceivedTokenContext = Context.GetInternalContext()->ReplicationSystem->GetNetTokenResolveContext(Context.GetLocalConnectionId());
+
 		// Dynamic objects can be referenced by either NetRefHandle or parent NetRefHandle and path combo
-		check(!bIsExported || ((CachedObject->NetRefHandle == NetRefHandle && CachedObject->RelativePath == RelativePath && CachedObject->OuterNetRefHandle == OuterRef.GetRefHandle()) || 
+		check(!bIsExported || ((CachedObject->NetRefHandle == NetRefHandle && CachedObject->RelativePath.IsEqualToToken(RelativePath, ReceivedTokenContext) && CachedObject->OuterNetRefHandle == OuterRef.GetRefHandle()) ||
 				(NetRefHandle.IsDynamic() && (NetRefHandle == CachedObject->NetRefHandle))));
 
 		OutObjectRef = FNetObjectReference(NetRefHandle);
@@ -1522,9 +1520,15 @@ void FObjectReferenceCache::ReadFullReferenceInternal(FNetSerializationContext&
 		// Register cached ref
 		FCachedNetObjectReference CachedObject;
 
+		FSharedStringNetToken::FTokenSource TokenSource =
+		{
+			.ReplicationSystemId = Context.GetInternalContext()->ReplicationSystem->GetId(),
+			.ConnectionId = Context.GetLocalConnectionId()
+		};
+
 		CachedObject.ObjectKey = nullptr;
 		CachedObject.NetRefHandle = NetRefHandle;
-		CachedObject.RelativePath = RelativePath;
+		CachedObject.RelativePath.AssignToken(RelativePath, TokenSource);
 		CachedObject.OuterNetRefHandle = OuterRef.GetRefHandle();
 		CachedObject.bIsPackage = NetRefHandle.IsStatic() && !OuterRef.GetRefHandle().IsValid();
 		CachedObject.bNoLoad = bNoLoad;
@@ -2123,10 +2127,8 @@ void FObjectReferenceCache::GenerateFullPath_r(FNetRefHandle RefHandle, const FN
 		}
 		else
 		{
-			FStringTokenStore* LocalStringTokenStore = ResolveContext.ReplicationSystem->GetNetTokenStore()->GetDataStore<FStringTokenStore>();
-
-			const TCHAR* ResolvedPath = LocalStringTokenStore->ResolveRemoteToken(CacheObject->RelativePath, *ResolveContext.RemoteNetTokenStoreState);
-			FullPath += FString::Printf(TEXT("[%s](%s)"), *RefHandle.ToString(), ResolvedPath);
+			const TCHAR* ResolvedPath = CacheObject->RelativePath.ResolveToken();
+			FullPath += FString::Printf(TEXT("[%s](%s)"), *RefHandle.ToString(), ResolvedPath ? ResolvedPath : TEXT("EMPTY"));
 		}
 	}
 }
@@ -2410,13 +2412,38 @@ const TCHAR* FObjectReferenceCache::GetObjectRelativePath(FNetRefHandle NetRefHa
 {
 	if (const FCachedNetObjectReference* CacheObject = ReferenceHandleToCachedReference.Find(NetRefHandle))
 	{
-		FStringTokenStore* LocalStringTokenStore = ReplicationSystem->GetNetTokenStore()->GetDataStore<FStringTokenStore>();
-		const TCHAR* ResolvedToken = LocalStringTokenStore->ResolveToken(CacheObject->RelativePath);
+		const TCHAR* ResolvedToken = CacheObject->RelativePath.ResolveToken();
 		return ResolvedToken;
 	}
 
 	return nullptr;
 }
 
+void FObjectReferenceCache::ReplicationSystemDestroyed(UReplicationSystem* RepSystem)
+{
+	// Invalidate FNetTokens that came from the RepSystem in our FCachedNetObjectReferences.
+	for (TPair<FNetRefHandle, FCachedNetObjectReference>& ReferencePair : ReferenceHandleToCachedReference)
+	{
+		FSharedStringNetToken& RelativePath = ReferencePair.Value.RelativePath;
+		if (RelativePath.GetSourceReplicationSystemId() == RepSystem->GetId())
+		{
+			RelativePath.MarkStale();
+		}
+	}
+}
+
+void FObjectReferenceCache::RemoveConnection(UReplicationSystem* RepSystem, uint32 ConnectionId)
+{
+	// Invalidate FNetTokens that came from the RepSystem and connection in our FCachedNetObjectReferences.
+	for (TPair<FNetRefHandle, FCachedNetObjectReference>& ReferencePair : ReferenceHandleToCachedReference)
+	{
+		FSharedStringNetToken& RelativePath = ReferencePair.Value.RelativePath;
+		if (RelativePath.GetSourceReplicationSystemId() == RepSystem->GetId() &&
+			RelativePath.GetSourceConnectionId() == ConnectionId)
+		{
+			RelativePath.MarkStale();
+		}
+	}
+}
 
 } // end namespace UE::Net::Private
diff --git a/Engine/Source/Runtime/Net/Iris/Private/Iris/ReplicationSystem/ObjectReferenceCache.h b/Engine/Source/Runtime/Net/Iris/Private/Iris/ReplicationSystem/ObjectReferenceCache.h
--- a/Engine/Source/Runtime/Net/Iris/Private/Iris/ReplicationSystem/ObjectReferenceCache.h
+++ b/Engine/Source/Runtime/Net/Iris/Private/Iris/ReplicationSystem/ObjectReferenceCache.h
@@ -5,6 +5,7 @@
 #include "CoreTypes.h"
 #include "Iris/Core/NetObjectReference.h"
 #include "Iris/ReplicationSystem/NetRefHandle.h"
+#include "Iris/ReplicationSystem/SharedStringNetToken.h"
 #include "Iris/ReplicationSystem/ReplicationSystemTypes.h"
 #include "UObject/WeakObjectPtr.h"
 #include "Containers/Map.h"
@@ -178,6 +179,12 @@ public:
 	/** Return the stored relative path of a replicated object if it's in the cache. */
 	const TCHAR* GetObjectRelativePath(FNetRefHandle NetRefHandle, UReplicationSystem* ReplicationSystem) const;
 
+	/** Called on the shared ObjectReferenceCache to invalidate cached data owned by the RepSystem (like FNetTokens) */
+	void ReplicationSystemDestroyed(UReplicationSystem* RepSystem);
+
+	/** Called on the shared ObjectReferenceCache to invalidate cached data owned by the connection (like remote FNetTokens) */
+	void RemoveConnection(UReplicationSystem* RepSystem, uint32 ConnectionId);
+
 private:
 	void AsyncPackageCallback(const FName& PackageName, UPackage* Package, EAsyncLoadingResult::Type Result);
 	void AsyncPackageForcedFailCallback(const FName& PackageName, UPackage* Package, EAsyncLoadingResult::Type Result);
@@ -191,7 +198,7 @@ private:
 		FNetRefHandle NetRefHandle;
 
 		// RelativePath to outer
-		FNetToken RelativePath;
+		FSharedStringNetToken RelativePath;
 
 		// Ref to outer
 		FNetRefHandle OuterNetRefHandle;
diff --git a/Engine/Source/Runtime/Net/Iris/Private/Iris/ReplicationSystem/ReplicationSystem.cpp b/Engine/Source/Runtime/Net/Iris/Private/Iris/ReplicationSystem/ReplicationSystem.cpp
--- a/Engine/Source/Runtime/Net/Iris/Private/Iris/ReplicationSystem/ReplicationSystem.cpp
+++ b/Engine/Source/Runtime/Net/Iris/Private/Iris/ReplicationSystem/ReplicationSystem.cpp
@@ -759,6 +759,11 @@ public:
 			ReplicationConditionals.RemoveConnection(ConnectionId);
 		}
 
+		{
+			FObjectReferenceCache& ObjectReferenceCache = ReplicationSystemInternal.GetObjectReferenceCache();
+			ObjectReferenceCache.RemoveConnection(ReplicationSystem, ConnectionId);
+		}
+
 		{
 			Connections.RemoveConnection(ConnectionId);
 		}
@@ -2286,6 +2291,11 @@ void FReplicationSystemFactory::DestroyReplicationSystem(UReplicationSystem* Sys
 		GetReplicationSystemDestroyedDelegate().Broadcast(System);
 	}
 
+	if (SharedObjectReferenceCache.IsValid())
+	{
+		SharedObjectReferenceCache->ReplicationSystemDestroyed(System);
+	}
+
 	System->Shutdown();
 	System->RemoveFromRoot();
 	System->MarkAsGarbage();
diff --git a/Engine/Source/Runtime/Net/Iris/Private/Iris/ReplicationSystem/SharedStringNetToken.cpp b/Engine/Source/Runtime/Net/Iris/Private/Iris/ReplicationSystem/SharedStringNetToken.cpp
new file mode 100644
--- /dev/null
+++ b/Engine/Source/Runtime/Net/Iris/Private/Iris/ReplicationSystem/SharedStringNetToken.cpp
@@ -0,0 +1,157 @@
+// Copyright Epic Games, Inc. All Rights Reserved.
+
+#include "SharedStringNetToken.h"
+#include "Iris/Serialization/NetSerializationContext.h"
+#include "Iris/ReplicationSystem/ReplicationSystem.h"
+#include "Iris/ReplicationSystem/StringTokenStore.h"
+#include "Iris/Serialization/InternalNetSerializationContext.h"
+
+namespace UE::Net::Private
+{
+
+void FSharedStringNetToken::AssignToken(FNetToken InToken, const FTokenSource InTokenSource)
+{
+	// In order to store a shared token we need to keep track of its resolve context.
+	checkf(InTokenSource.ReplicationSystemId != InvalidReplicationSystemId, TEXT("FSharedStringNetToken::AssignToken requires a valid ReplicationSystemId."));
+
+	Token = InToken;
+	TokenSource = InTokenSource;
+}
+
+void FSharedStringNetToken::MarkStale()
+{
+	TokenSource = FTokenSource();
+}
+
+const TCHAR* FSharedStringNetToken::ResolveToken() const
+{
+	// Can only resolve this token if we know which TokenStore it came from
+	if (!ensureMsgf(!IsStale(), TEXT("FSharedStringNetToken::ResolveToken: stored token is stale, it can't be resolved.")))
+	{
+		return nullptr;
+	}
+
+	const UReplicationSystem* ReplicationSystem = GetReplicationSystem(TokenSource.ReplicationSystemId);
+
+	if (!ensureMsgf(ReplicationSystem, TEXT("FSharedStringNetToken::ResolveToken: source ReplicationSystemId is invalid.")))
+	{
+		return nullptr;
+	}
+
+	const FNetTokenStore* const NetTokenStore = ReplicationSystem->GetNetTokenStore();
+	const FStringTokenStore* const StringTokenStore = NetTokenStore->GetDataStore<FStringTokenStore>();
+	const FNetTokenStoreState* const RemoteNetTokenStoreState = NetTokenStore->GetRemoteNetTokenStoreState(TokenSource.ConnectionId);
+
+	return StringTokenStore->ResolveToken(Token, RemoteNetTokenStoreState);
+}
+
+bool FSharedStringNetToken::IsEqualToToken(const FNetToken& OtherToken, const FNetTokenResolveContext& OtherResolveContext) const
+{
+	// Can only compare this token if we know which TokenStore it came from
+	if (!ensureMsgf(!IsStale(), TEXT("FSharedStringNetToken::IsEqualToToken: stored token is stale, it can't be compared.")))
+	{
+		return false;
+	}
+
+	// An invalid token doesn't need a resolve context
+	if (!OtherToken.IsValid() && !Token.IsValid())
+	{
+		return true;
+	}
+
+	if (!ensureMsgf(OtherResolveContext.NetTokenStore, TEXT("FSharedStringNetToken::IsEqualToToken: OtherResolveContext requires a valid NetTokenStore.")))
+	{
+		return false;
+	}
+
+	const UReplicationSystem* ReplicationSystem = GetReplicationSystem(TokenSource.ReplicationSystemId);
+
+	if (!ensureMsgf(ReplicationSystem, TEXT("FSharedStringNetToken::IsEqualToToken: source ReplicationSystemId %u is invalid."), TokenSource.ReplicationSystemId))
+	{
+		return false;
+	}
+
+	// Can only compare token values directly if they're from the same token store
+	const FNetTokenStore* const ThisNetTokenStore = ReplicationSystem->GetNetTokenStore();
+	const FStringTokenStore* const ThisStringTokenStore = ThisNetTokenStore->GetDataStore<FStringTokenStore>();
+	const FStringTokenStore* const OtherStringTokenStore = OtherResolveContext.NetTokenStore->GetDataStore<FStringTokenStore>();
+
+	if (ThisStringTokenStore == OtherStringTokenStore)
+	{
+		return Token == OtherToken;
+	}
+	else
+	{
+		// Different token stores: resolve the tokens and compare the strings.
+		const FNetTokenStoreState* const ThisRemoteNetTokenStoreState = ThisNetTokenStore->GetRemoteNetTokenStoreState(TokenSource.ConnectionId);
+		const TCHAR* const ThisResolvedToken = ThisStringTokenStore->ResolveToken(Token, ThisRemoteNetTokenStoreState);
+		const TCHAR* const OtherResolvedToken = OtherStringTokenStore->ResolveToken(OtherToken, OtherResolveContext.RemoteNetTokenStoreState);
+
+		if (ThisResolvedToken && OtherResolvedToken)
+		{
+			return FCString::Strcmp(ThisResolvedToken, OtherResolvedToken) == 0;
+		}
+		else
+		{
+			return ThisResolvedToken == OtherResolvedToken;
+		}
+	}
+}
+
+void FSharedStringNetToken::WriteNetToken(FNetSerializationContext& Context) const
+{
+	if (!ensureMsgf(!IsStale(), TEXT("FSharedStringNetToken::WriteNetToken: stored token is stale, it can't be written.")))
+	{
+		Context.SetError(GNetError_InvalidValue);
+		return;
+	}
+
+	const uint32 ContextRepSystemId = Context.GetInternalContext()->ReplicationSystem->GetId();
+
+	if (!ensureMsgf(ContextRepSystemId == TokenSource.ReplicationSystemId, TEXT("FSharedStringNetToken::WriteNetToken: mismatched replication systems! This ID: %u, context ID: %u"), TokenSource.ReplicationSystemId, ContextRepSystemId))
+	{
+		Context.SetError(GNetError_InvalidValue);
+		return;
+	}
+
+	const FStringTokenStore* const ContextStringTokenStore = Context.GetNetTokenStore()->GetDataStore<FStringTokenStore>();
+
+	ContextStringTokenStore->WriteNetToken(Context, Token);
+}
+
+void FSharedStringNetToken::ConditionalWriteNetTokenData(FNetSerializationContext& Context, Private::FNetExportContext* ExportContext) const
+{
+	if (!ensureMsgf(!IsStale(), TEXT("FSharedStringNetToken::ConditionalWriteNetTokenData: stored token is stale, it can't be written.")))
+	{
+		Context.SetError(GNetError_InvalidValue);
+		return;
+	}
+
+	const uint32 ContextRepSystemId = Context.GetInternalContext()->ReplicationSystem->GetId();
+
+	if (!ensureMsgf(ContextRepSystemId == TokenSource.ReplicationSystemId, TEXT("FSharedStringNetToken::ConditionalWriteNetTokenData: mismatched replication systems! This ID: %u, context ID: %u"), TokenSource.ReplicationSystemId, ContextRepSystemId))
+	{
+		Context.SetError(GNetError_InvalidValue);
+		return;
+	}
+
+	Context.GetNetTokenStore()->ConditionalWriteNetTokenData(Context, ExportContext, Token);
+}
+
+
+FString FSharedStringNetToken::ToString() const
+{
+	FString Result;
+	if (IsStale())
+	{
+		Result = TEXT("Stale SharedStringNetToken");
+	}
+	else
+	{
+		Result = FString::Printf(TEXT("SharedStringNetToken (Auth:%u TypeId=%u Index=%u RepSystemId=%u ConnectionId=%u)"),
+			Token.IsAssignedByAuthority(), Token.GetTypeId(), Token.GetIndex(), TokenSource.ReplicationSystemId, TokenSource.ConnectionId);
+	}
+	return Result;
+}
+
+}
\ No newline at end of file
diff --git a/Engine/Source/Runtime/Net/Iris/Private/Iris/ReplicationSystem/SharedStringNetToken.h b/Engine/Source/Runtime/Net/Iris/Private/Iris/ReplicationSystem/SharedStringNetToken.h
new file mode 100644
--- /dev/null
+++ b/Engine/Source/Runtime/Net/Iris/Private/Iris/ReplicationSystem/SharedStringNetToken.h
@@ -0,0 +1,110 @@
+// Copyright Epic Games, Inc. All Rights Reserved.
+
+#pragma once
+
+#include "HAL/Platform.h"
+#include "Iris/IrisConstants.h"
+#include "Net/Core/NetToken/NetToken.h"
+
+namespace UE::Net
+{
+	class FNetSerializationContext;
+
+	namespace Private
+	{
+		class FNetExportContext;
+	}
+}
+
+namespace UE::Net::Private
+{
+
+/**
+ * Representation of a string FNetToken that's safe to use in a context that's shared between
+ * multiple replication systems, like a shared ObjectReferenceCache.
+ * 
+ * Stores an FNetToken, the replication system it was sourced from, and if it was assigned remotely,
+ * also the connection ID it was sourced from. These IDs are used to look up the correct FNetTokenStore
+ * and remote FNetTokenStoreState (if needed) to resolve or compare the token. The public API enforces that
+ * the token will only be resolved/compared using the correct source token store & state.
+ * 
+ * When the ReplicationSystem that owns the source token store is destroyed, this shared token
+ * must be marked as stale. This is handled in FReplicationSystemFactory::DestroyReplicationSystem.
+ */
+class FSharedStringNetToken
+{
+public:
+	struct FTokenSource
+	{
+		/** The ReplicationSystem that assigned this token. Used to look up the corresponding FStringTokenStore. */
+		uint32 ReplicationSystemId = InvalidReplicationSystemId;
+
+		/** The connection within the replication system that sent the token, if it was remotely assigned. Used to look up the corresponding FNetTokenStoreState. */
+		uint32 ConnectionId = InvalidConnectionId;
+	};
+
+	/** Assigns a token. Remote tokens need to specify a valid ConnectionId in the FTokenSource. */
+	void AssignToken(FNetToken InToken, const FTokenSource InTokenSource);
+
+	/**
+	 * Shared tokens become meaningless without their original ResolveContext.
+	 * This token must be invalidated when the ResolveContext is destroyed.
+	 * Called via FReplicationSystemFactory::DestroyReplicationSystem and FObjectReferenceCache::ReplicationSystemDestroyed.
+	 */
+	void MarkStale();
+
+	/** Returns true if the stored token is valid and its backing store is also valid/not stale. */
+	bool IsValid() const
+	{
+		return Token.IsValid() && !IsStale();
+	}
+
+	/** Returns true if the stored token's backing store is stale/has been destroyed. */
+	bool IsStale() const
+	{
+		return TokenSource.ReplicationSystemId == InvalidReplicationSystemId;
+	}
+
+	/** Resolves the token if it's not stale using its source token store and if needed, the remote token store state. Returns null if the token is stale. */
+	const TCHAR* ResolveToken() const;
+
+	/**
+	 * Compare this token to another, using their respective backing token stores if necessary.
+	 * If both tokens come from the same store, the tokens can be compared directly.
+	 * If the tokens come from different stores, we have to resolve the strings and compare those.
+	 * Stale shared tokens are not considered equal to any (non-shared) NetToken.
+	 * If OtherToken is invalid, it's considered equal only if this token is also invalid and not stale.
+	 */
+	bool IsEqualToToken(const FNetToken& OtherToken, const FNetTokenResolveContext& OtherResolveContext) const;
+
+	/** Return a string representation of this token. */
+	FString ToString() const;
+
+	/** Returns the replication system that assigned this token. */
+	uint32 GetSourceReplicationSystemId() const
+	{
+		return TokenSource.ReplicationSystemId;
+	}
+
+	/** Returns the connection ID that assigned this token. It is only relevant together with the source replication system ID. */
+	uint32 GetSourceConnectionId() const
+	{
+		return TokenSource.ConnectionId;
+	}
+
+	/**
+	 * Forwards to WriteNetToken of the token store of the FNetSerializationContext. Validates that the context's store is the same as this token's store.
+	 */
+	void WriteNetToken(UE::Net::FNetSerializationContext& Context) const;
+
+	/**
+	 * Forwards to ConditionalWriteNetTokenData of the token store of the FNetSerializationContext. Validates that the context's store is the same as the token's store.
+	 */
+	void ConditionalWriteNetTokenData(FNetSerializationContext& Context, Private::FNetExportContext* ExportContext) const;
+
+private:
+	FNetToken Token;
+	FTokenSource TokenSource;
+};
+
+}
diff --git a/Engine/Source/Runtime/Net/Iris/Public/Iris/ReplicationSystem/NetTokenStore.h b/Engine/Source/Runtime/Net/Iris/Public/Iris/ReplicationSystem/NetTokenStore.h
--- a/Engine/Source/Runtime/Net/Iris/Public/Iris/ReplicationSystem/NetTokenStore.h
+++ b/Engine/Source/Runtime/Net/Iris/Public/Iris/ReplicationSystem/NetTokenStore.h
@@ -135,7 +135,7 @@ public:
 	// Serialization methods for NetTokens that does not include the TypeId as this is known by the NetTokenDataStore.
 
 	// Write NetToken
-	IRISCORE_API void WriteNetToken(UE::Net::FNetSerializationContext& Context, FNetToken Token);
+	IRISCORE_API void WriteNetToken(UE::Net::FNetSerializationContext& Context, FNetToken Token) const;
 
 	// Read NetToken
 	IRISCORE_API FNetToken ReadNetToken(UE::Net::FNetSerializationContext& Context);
