Iris Parallelization - Added thread safety functionality to DeltaCompressionBaselineManager. Added FScopedBaselineAccess which guards parallel access to the BaselineManager, which is used in the ReplicationWriter in preparation for running it in parallel. Also renamed GetIsInParallelPhase to IsInParallelPhase to align with Iris Beta naming conventions.

#tests Built and tested against Dedicated Server Cooked data and PIE
[RN] Added thread safety to DeltaCompressionBaselineManager

diff --git a/Engine/Source/Runtime/Net/Iris/Private/Iris/ReplicationSystem/DeltaCompression/DeltaCompressionBaselineManager.cpp b/Engine/Source/Runtime/Net/Iris/Private/Iris/ReplicationSystem/DeltaCompression/DeltaCompressionBaselineManager.cpp
--- a/Engine/Source/Runtime/Net/Iris/Private/Iris/ReplicationSystem/DeltaCompression/DeltaCompressionBaselineManager.cpp
+++ b/Engine/Source/Runtime/Net/Iris/Private/Iris/ReplicationSystem/DeltaCompression/DeltaCompressionBaselineManager.cpp
@@ -389,6 +389,7 @@ void FDeltaCompressionBaselineManager::FreeAllPerObjectInfos()
 
 FDeltaCompressionBaseline FDeltaCompressionBaselineManager::CreateBaseline(uint32 ConnId, uint32 ObjectIndex, uint32 BaselineIndex)
 {
+	IRIS_PROFILER_SCOPE(FDeltaCompressionBaselineManager_CreateBaseline);
 	UE_NET_DC_BASELINE_CHECK(BaselineSharingContext.ObjectInfoIndicesWithNewBaseline.GetNumBits() > 0);
 	
 	FDeltaCompressionBaseline Baseline;
@@ -566,6 +567,15 @@ void FDeltaCompressionBaselineManager::ReleaseBaselinesForConnection(uint32 Conn
 
 void FDeltaCompressionBaselineManager::ConstructBaselineSharingContext(FBaselineSharingContext& SharingContext)
 {
+	IRIS_PROFILER_SCOPE(FDeltaCompressionBaselineManager_ConstructBaselineSharingContext);
+
+#if UE_SUPPORT_PARALLEL_IRIS
+	if (ReplicationSystem)
+	{
+		ensureMsgf(!ReplicationSystem->IsInParallelPhase(), TEXT("ConstructBaselineSharingContext should never be called from ParallelPhase"));
+	}
+#endif // UE_SUPPORT_PARALLEL_IRIS
+
 	SharingContext.CreatedBaselineCount = 0;
 	SharingContext.ObjectInfoIndicesWithNewBaseline.Init(UsedPerObjectInfos.GetNumBits());
 	SharingContext.ObjectInfoIndexToBaselineInfoIndex.SetNumZeroed(UsedPerObjectInfos.GetNumBits());
@@ -573,6 +583,15 @@ void FDeltaCompressionBaselineManager::ConstructBaselineSharingContext(FBaseline
 
 void FDeltaCompressionBaselineManager::DestructBaselineSharingContext(FBaselineSharingContext& SharingContext)
 {
+	IRIS_PROFILER_SCOPE(FDeltaCompressionBaselineManager_DestructBaselineSharingContext);
+
+#if UE_SUPPORT_PARALLEL_IRIS
+	if (ReplicationSystem)
+	{
+		ensureMsgf(!ReplicationSystem->IsInParallelPhase(), TEXT("DestructBaselineSharingContext should never be called from ParallelPhase"));
+	}
+#endif // UE_SUPPORT_PARALLEL_IRIS
+
 	auto ReleaseBaseline = [this, SharingContext](uint32 ObjectInfoIndex)
 	{
 		const DeltaCompressionBaselineStateInfoIndexType BaselineStateInfoIndex = SharingContext.ObjectInfoIndexToBaselineInfoIndex[ObjectInfoIndex];
@@ -741,6 +760,51 @@ FString FDeltaCompressionBaselineManager::PrintDeltaCompressionStatus(uint32 Con
 	}
 }
 
+#if UE_SUPPORT_PARALLEL_IRIS
+void FDeltaCompressionBaselineManager::AcquireBaselineManager(const bool bWriteLock)
+{
+	IRIS_PROFILER_SCOPE(FDeltaCompressionBaselineManager_AcquireBaselineManager);
+	if (bWriteLock)
+	{
+		DeltaCompressionBaselineManagerMutex.WriteLock();
+	}
+	else
+	{
+		DeltaCompressionBaselineManagerMutex.ReadLock();
+	}
+}
+
+void FDeltaCompressionBaselineManager::ReleaseBaselineManager(const bool bWriteLock)
+{
+	IRIS_PROFILER_SCOPE(FDeltaCompressionBaselineManager_ReleaseBaselineManager);
+	if (bWriteLock)
+	{
+		DeltaCompressionBaselineManagerMutex.WriteUnlock();
+	}
+	else
+	{
+		DeltaCompressionBaselineManagerMutex.ReadUnlock();
+	}
+}
+#endif // UE_SUPPORT_PARALLEL_IRIS
+
+FScopedBaselineAccess::FScopedBaselineAccess(FDeltaCompressionBaselineManager* InBaselineManager, const bool bWriteLock)
+{
+#if UE_SUPPORT_PARALLEL_IRIS
+	check(InBaselineManager);
+	InBaselineManager->AcquireBaselineManager(bWriteLock);
+	StoredBaselineManager = InBaselineManager;
+	bStoredWriteLockType = bWriteLock;
+#endif // UE_SUPPORT_PARALLEL_IRIS
+}
+
+FScopedBaselineAccess::~FScopedBaselineAccess()
+{
+#if UE_SUPPORT_PARALLEL_IRIS
+	StoredBaselineManager->ReleaseBaselineManager(bStoredWriteLockType);
+#endif // UE_SUPPORT_PARALLEL_IRIS
+}
+
 }
 
 #undef UE_NET_DC_BASELINE_CHECK
diff --git a/Engine/Source/Runtime/Net/Iris/Private/Iris/ReplicationSystem/DeltaCompression/DeltaCompressionBaselineManager.h b/Engine/Source/Runtime/Net/Iris/Private/Iris/ReplicationSystem/DeltaCompression/DeltaCompressionBaselineManager.h
--- a/Engine/Source/Runtime/Net/Iris/Private/Iris/ReplicationSystem/DeltaCompression/DeltaCompressionBaselineManager.h
+++ b/Engine/Source/Runtime/Net/Iris/Private/Iris/ReplicationSystem/DeltaCompression/DeltaCompressionBaselineManager.h
@@ -4,8 +4,10 @@
 
 #include "Containers/Array.h"
 #include "Net/Core/NetBitArray.h"
+#include "Iris/IrisConfig.h"
 #include "Iris/ReplicationSystem/DeltaCompression/DeltaCompressionBaseline.h"
 #include "Iris/ReplicationSystem/DeltaCompression/DeltaCompressionBaselineStorage.h"
+#include "Misc/TransactionallySafeRWLock.h"
 
 class UReplicationSystem;
 
@@ -21,6 +23,7 @@ namespace UE::Net::Private
 	typedef uint32 FInternalNetRefIndex;
 	class FNetRefHandleManager;
 	class FReplicationConnections;
+	class FScopedBaselineAccess;
 }
 
 namespace UE::Net::Private
@@ -171,6 +174,17 @@ private:
 		uint32 CreatedBaselineCount = 0;
 	};
 
+#if UE_SUPPORT_PARALLEL_IRIS
+	friend FScopedBaselineAccess;
+
+	// For internal use by the FScopedBaselineAccess to lock DeltaCompressionBaselineManagerMutex
+	void AcquireBaselineManager(const bool bWriteLock);
+	// For internal use by the FScopedBaselineAccess to unlock DeltaCompressionBaselineManagerMutex
+	void ReleaseBaselineManager(const bool bWriteLock);
+
+	FTransactionallySafeRWLock DeltaCompressionBaselineManagerMutex;
+#endif // UE_SUPPORT_PARALLEL_IRIS
+
 private:
 
 	void UpdateScope();
@@ -246,6 +260,23 @@ private:
 	FBaselineSharingContext BaselineSharingContext;
 };
 
+/**
+* A helper RAII class to allow thread-safe access to Baselines
+* When making calls to the FDeltaCompressionBaselineManager in a parallel phase, first create FScopedBaselineAccess to
+* ensure that the access is threadsafe during the lifetime of the FScopedBaselineAccess object.
+*/
+class FScopedBaselineAccess
+{
+public:
+	
+	explicit FScopedBaselineAccess(FDeltaCompressionBaselineManager* InBaselineManager, const bool bWriteLock);
+
+	~FScopedBaselineAccess();
+
+private:
+	FDeltaCompressionBaselineManager* StoredBaselineManager = nullptr;
+	bool bStoredWriteLockType = false;
+};
 
 inline uint32 FDeltaCompressionBaselineManager::GetMaxDeltaCompressedObjectCount() const
 {
diff --git a/Engine/Source/Runtime/Net/Iris/Private/Iris/ReplicationSystem/Polling/ObjectPoller.cpp b/Engine/Source/Runtime/Net/Iris/Private/Iris/ReplicationSystem/Polling/ObjectPoller.cpp
--- a/Engine/Source/Runtime/Net/Iris/Private/Iris/ReplicationSystem/Polling/ObjectPoller.cpp
+++ b/Engine/Source/Runtime/Net/Iris/Private/Iris/ReplicationSystem/Polling/ObjectPoller.cpp
@@ -96,7 +96,7 @@ namespace CVars
 
 		void DoTask() const
 		{
-			check(ReplicationSystemInternal->GetIsInParallelPhase());
+			check(ReplicationSystemInternal->IsInParallelPhase());
 
 			IRIS_CSV_PROFILER_SCOPE(Iris, ReplicationPollAndCopyTask);
 
@@ -288,7 +288,7 @@ void FObjectPoller::PollAndCopySingleObject(FInternalNetRefIndex ObjectIndex)
 	DirtyObjectsThisFrame = DirtyObjectsAccessor.GetDirtyNetObjects();
 
 	//We don't run this PollAndCopySingleObject during the parallel Poll phase, so we can use the parent context NetStatsContext and PollStats
-	check(!ReplicationSystemInternal->GetIsInParallelPhase());
+	check(!ReplicationSystemInternal->IsInParallelPhase());
 
 	IRIS_PROFILER_SCOPE_VERBOSE(ForcePollAndCopy);
 	ForcePollObject(ObjectIndex, NetStatsContext, PollStats);
diff --git a/Engine/Source/Runtime/Net/Iris/Private/Iris/ReplicationSystem/ReplicationSystem.cpp b/Engine/Source/Runtime/Net/Iris/Private/Iris/ReplicationSystem/ReplicationSystem.cpp
--- a/Engine/Source/Runtime/Net/Iris/Private/Iris/ReplicationSystem/ReplicationSystem.cpp
+++ b/Engine/Source/Runtime/Net/Iris/Private/Iris/ReplicationSystem/ReplicationSystem.cpp
@@ -909,6 +909,22 @@ const UE::Net::Private::FReplicationSystemInternal* UReplicationSystem::GetRepli
 	return &Impl->ReplicationSystemInternal;
 }
 
+void UReplicationSystem::StartParallelPhase()
+{ 
+	GetReplicationSystemInternal()->StartParallelPhase();
+}
+
+void UReplicationSystem::StopParallelPhase()
+{
+	GetReplicationSystemInternal()->StopParallelPhase(); 
+}
+
+bool UReplicationSystem::IsInParallelPhase() const
+{
+	return GetReplicationSystemInternal()->IsInParallelPhase();
+}
+
+
 void UReplicationSystem::NetUpdate(float DeltaSeconds)
 {
 	using namespace UE::Net;
diff --git a/Engine/Source/Runtime/Net/Iris/Private/Iris/ReplicationSystem/ReplicationSystemInternal.h b/Engine/Source/Runtime/Net/Iris/Private/Iris/ReplicationSystem/ReplicationSystemInternal.h
--- a/Engine/Source/Runtime/Net/Iris/Private/Iris/ReplicationSystem/ReplicationSystemInternal.h
+++ b/Engine/Source/Runtime/Net/Iris/Private/Iris/ReplicationSystem/ReplicationSystemInternal.h
@@ -169,7 +169,7 @@ public:
 		TypeStats.SetIsInParallelPhase(false); 
 	}
 
-	bool GetIsInParallelPhase() const
+	bool IsInParallelPhase() const
 	{
 		return bIsInParallelPhase;
 	}
diff --git a/Engine/Source/Runtime/Net/Iris/Private/Iris/ReplicationSystem/ReplicationWriter.cpp b/Engine/Source/Runtime/Net/Iris/Private/Iris/ReplicationSystem/ReplicationWriter.cpp
--- a/Engine/Source/Runtime/Net/Iris/Private/Iris/ReplicationSystem/ReplicationWriter.cpp
+++ b/Engine/Source/Runtime/Net/Iris/Private/Iris/ReplicationSystem/ReplicationWriter.cpp
@@ -634,14 +634,20 @@ void FReplicationWriter::StopReplication(uint32 InternalIndex)
 	Attachments.DropAllAttachments(ENetObjectAttachmentType::Normal, InternalIndex);
 
 	// Release baselines
-	if (Info.PendingBaselineIndex != FDeltaCompressionBaselineManager::InvalidBaselineIndex)
-	{
-		BaselineManager->DestroyBaseline(Parameters.ConnectionId, InternalIndex, Info.PendingBaselineIndex);
-	}
-	if (Info.PendingBaselineIndex != Info.LastAckedBaselineIndex && Info.LastAckedBaselineIndex != FDeltaCompressionBaselineManager::InvalidBaselineIndex)
-	{
-		BaselineManager->DestroyBaseline(Parameters.ConnectionId, InternalIndex, Info.LastAckedBaselineIndex);
-	}
+    const bool bDestroyPendingBaseline = Info.PendingBaselineIndex != FDeltaCompressionBaselineManager::InvalidBaselineIndex;
+    const bool bDestroyLastAckedBaseline = Info.PendingBaselineIndex != Info.LastAckedBaselineIndex && Info.LastAckedBaselineIndex != FDeltaCompressionBaselineManager::InvalidBaselineIndex;
+    if (bDestroyPendingBaseline || bDestroyLastAckedBaseline)
+    {
+        UE_IRIS_PARALLEL_EXPR(FScopedBaselineAccess BaselineRW(BaselineManager, true));
+        if (bDestroyPendingBaseline)
+        {
+            BaselineManager->DestroyBaseline(Parameters.ConnectionId, InternalIndex, Info.PendingBaselineIndex);
+        }
+        if (bDestroyLastAckedBaseline)
+        {
+            BaselineManager->DestroyBaseline(Parameters.ConnectionId, InternalIndex, Info.LastAckedBaselineIndex);
+        }
+    }
 
 	Info.PendingBaselineIndex = Info.LastAckedBaselineIndex = FDeltaCompressionBaselineManager::InvalidBaselineIndex;
 
@@ -966,6 +972,8 @@ void FReplicationWriter::UpdateScope(const FNetBitArrayView& UpdatedScope)
 
 void FReplicationWriter::UpdateDirtyGlobalLifetimeConditionals(TArrayView<FInternalNetRefIndex> ObjectsWithDirtyConditionals)
 {
+	ensureMsgf(!Parameters.ReplicationSystem->GetReplicationSystemInternal()->IsInParallelPhase(), TEXT("UpdateDirtyGlobalLifetimeConditionals  should never be called from within a parallel phase"));
+
 	for (FInternalNetRefIndex InternalObjectIndex : ObjectsWithDirtyConditionals)
 	{
 		// Better safe than sorry, we do not want to dirty something going out of scope.
@@ -2553,6 +2561,7 @@ FReplicationWriter::EWriteObjectStatus FReplicationWriter::WriteObjectAndSubObje
 	{
 		if (CreatedBaselineIndex != FDeltaCompressionBaselineManager::InvalidBaselineIndex)
 		{
+			UE_IRIS_PARALLEL_EXPR(FScopedBaselineAccess BaselineRW(BaselineManager, true));
 			UE_LOG_REPWRITER(VeryVerbose, TEXT("Destroy cancelled baseline %u for ( InternalIndex: %u )"), CreatedBaselineIndex, InternalIndex);
 			BaselineManager->DestroyBaseline(Parameters.ConnectionId, InternalIndex, CreatedBaselineIndex);
 		}
@@ -2660,6 +2669,7 @@ FReplicationWriter::EWriteObjectStatus FReplicationWriter::WriteObjectAndSubObje
 			// If the last transmitted baseline is acknowledged we can request a new baseline to be stored for the current state, we cannot compress against it until it has been acknowledged
 			if (Info.IsDeltaCompressionEnabled)
 			{
+				UE_IRIS_PARALLEL_EXPR(FScopedBaselineAccess BaselineRW(BaselineManager, true));
 				// Lookup current baseline that we should compress against
 				if (Info.LastAckedBaselineIndex != FDeltaCompressionBaselineManager::InvalidBaselineIndex)
 				{
@@ -2767,6 +2777,7 @@ FReplicationWriter::EWriteObjectStatus FReplicationWriter::WriteObjectAndSubObje
 			{
 				if (Info.IsDeltaCompressionEnabled)
 				{
+					UE_IRIS_PARALLEL_EXPR(FScopedBaselineAccess BaselineReadOnly(BaselineManager, false));
 					SerializeObjectStateDelta(Context, InternalIndex, Info, ObjectData, ReplicatedObjectStateBuffer, CurrentBaseline, CreatedBaselineIndex);
 				}
 				else
@@ -3970,6 +3981,7 @@ FReplicationWriter::EWriteObjectRetryMode FReplicationWriter::HandleObjectBatchF
 		// If we did not end up using the baseline we need to release it
 		if (BatchObjectInfo.bSentState && BatchObjectInfo.NewBaselineIndex != FDeltaCompressionBaselineManager::InvalidBaselineIndex)
 		{
+			UE_IRIS_PARALLEL_EXPR(FScopedBaselineAccess BaselineRW(BaselineManager, true));
 			BaselineManager->LostBaseline(Parameters.ConnectionId, BatchObjectInfo.InternalIndex, BatchObjectInfo.NewBaselineIndex);
 		}
 		
diff --git a/Engine/Source/Runtime/Net/Iris/Public/Iris/ReplicationSystem/ReplicationSystem.h b/Engine/Source/Runtime/Net/Iris/Public/Iris/ReplicationSystem/ReplicationSystem.h
--- a/Engine/Source/Runtime/Net/Iris/Public/Iris/ReplicationSystem/ReplicationSystem.h
+++ b/Engine/Source/Runtime/Net/Iris/Public/Iris/ReplicationSystem/ReplicationSystem.h
@@ -795,6 +795,10 @@ public:
 
 	IRISCORE_API UE::Net::FReplicationSystemDelegates& GetDelegates();
 
+	IRISCORE_API void StartParallelPhase();
+	IRISCORE_API void StopParallelPhase();
+	IRISCORE_API bool IsInParallelPhase() const;
+
 	/** Function to query if a particular replicated object is relevant to a connection or not. */
 	UE_EXPERIMENTAL(5.7, "IsRelevantForConnection() is experimental and should not be used.")
 	IRISCORE_API bool IsRelevantForConnection(FNetRefHandle NetRefHandle, uint32 ConnectionId) const;
