Server NetDrivers run Iris replication in a transaction if necessary.

diff --git a/Engine/Source/Runtime/Engine/Classes/Engine/NetDriver.h b/Engine/Source/Runtime/Engine/Classes/Engine/NetDriver.h
--- a/Engine/Source/Runtime/Engine/Classes/Engine/NetDriver.h
+++ b/Engine/Source/Runtime/Engine/Classes/Engine/NetDriver.h
@@ -2462,6 +2462,8 @@ private:
 	void UpdateIrisReplicationViews() const;
 	void SendClientMoveAdjustments();
 	void PostDispatchSendUpdate();
+	void InternalIrisUpdateTransactional(float DeltaSeconds);
+	void InternalIrisUpdate(float DeltaSeconds);
 
 	/** Tell the registered NetAnalytics to send their analytics via the provider */
 	void SendNetAnalytics();
diff --git a/Engine/Source/Runtime/Engine/Private/NetDriver.cpp b/Engine/Source/Runtime/Engine/Private/NetDriver.cpp
--- a/Engine/Source/Runtime/Engine/Private/NetDriver.cpp
+++ b/Engine/Source/Runtime/Engine/Private/NetDriver.cpp
@@ -1148,9 +1148,7 @@ void UNetDriver::TickFlush(float DeltaSeconds)
 
 		if (ReplicationSystem)
 		{
-			UpdateIrisReplicationViews();
-			SendClientMoveAdjustments();
-			ReplicationSystem->NetUpdate(DeltaSeconds);
+			InternalIrisUpdateTransactional(DeltaSeconds);
 		}
 		else if (ClientConnections.Num() > 0)
 		{
@@ -1245,6 +1243,31 @@ void UNetDriver::TickFlush(float DeltaSeconds)
 	UpdateNetworkLagState();
 }
 
+void UNetDriver::InternalIrisUpdateTransactional(float DeltaSeconds)
+{
+#if UE_WITH_REMOTE_OBJECT_HANDLE && UE_AUTORTFM
+	if (bReplicateTransactionally)
+	{
+		const FName TransactionalWorkName("UNetDriver::InternalIrisUpdate");
+		UE::RemoteExecutor::ExecuteTransactional(TransactionalWorkName, [this, DeltaSeconds]()
+			{
+				InternalIrisUpdate(DeltaSeconds);
+			});
+	}
+	else
+#endif
+	{
+		InternalIrisUpdate(DeltaSeconds);
+	}
+}
+
+void UNetDriver::InternalIrisUpdate(float DeltaSeconds)
+{
+	UpdateIrisReplicationViews();
+	SendClientMoveAdjustments();
+	ReplicationSystem->NetUpdate(DeltaSeconds);
+}
+
 void UNetDriver::UpdateUnmappedObjects()
 {
 	CSV_SCOPED_TIMING_STAT(Networking, UpdateUnmappedObjects);
diff --git a/Engine/Source/Runtime/Net/Iris/Private/Iris/ReplicationSystem/ReplicationSystem.cpp b/Engine/Source/Runtime/Net/Iris/Private/Iris/ReplicationSystem/ReplicationSystem.cpp
--- a/Engine/Source/Runtime/Net/Iris/Private/Iris/ReplicationSystem/ReplicationSystem.cpp
+++ b/Engine/Source/Runtime/Net/Iris/Private/Iris/ReplicationSystem/ReplicationSystem.cpp
@@ -563,7 +563,7 @@ public:
 	void ExecuteQuantizeDirtyStateData(bool bTransact)
 	{
 #if UE_WITH_REMOTE_OBJECT_HANDLE && UE_AUTORTFM
-		if (bTransact)
+		if (bTransact && !AutoRTFM::IsTransactional())
 		{
 			static FName TransactionalWorkName("Iris FReplicationSystemImpl::QuantizeDirtyStateData");
 			UE::RemoteExecutor::ExecuteTransactional(TransactionalWorkName, [this]()
diff --git a/Engine/Source/Runtime/Net/Iris/Private/Iris/Stats/NetStats.cpp b/Engine/Source/Runtime/Net/Iris/Private/Iris/Stats/NetStats.cpp
--- a/Engine/Source/Runtime/Net/Iris/Private/Iris/Stats/NetStats.cpp
+++ b/Engine/Source/Runtime/Net/Iris/Private/Iris/Stats/NetStats.cpp
@@ -43,7 +43,7 @@ CSV_DEFINE_CATEGORY(IrisWriteExportsCount, WITH_SERVER_CODE);
 void FNetSendStats::Accumulate(const FNetSendStats& Other)
 {
 #if UE_NET_IRIS_CSV_STATS
-	FScopeLock Lock(&CS);
+	UE::TScopeLock Lock(CS);
 
 	Stats.ScheduledForReplicationRootObjectCount += Other.Stats.ScheduledForReplicationRootObjectCount;
 	Stats.ReplicatedRootObjectCount += Other.Stats.ReplicatedRootObjectCount;
@@ -63,7 +63,7 @@ void FNetSendStats::Accumulate(const FNetSendStats& Other)
 void FNetSendStats::Reset()
 {
 #if UE_NET_IRIS_CSV_STATS
-	FScopeLock Lock(&CS);
+	UE::TScopeLock Lock(CS);
 	Stats = FStats();
 #endif
 }
@@ -71,7 +71,7 @@ void FNetSendStats::Reset()
 void FNetSendStats::ReportCsvStats()
 {
 #if CSV_PROFILER_STATS
-	FScopeLock Lock(&CS);
+	UE::TScopeLock Lock(CS);
 
 	// Calculate connection averages for some stats
 	if (Stats.ReplicatingConnectionCount > 0)
@@ -195,7 +195,7 @@ FNetStatsContext* FNetTypeStats::CreateChildNetStatsContext()
 {
 	IRIS_PROFILER_SCOPE(FNetTypeStats_CreateChildNetStatsContext);
 
-	FScopeLock Lock(&ChildStatsContextCS);
+	UE::TScopeLock Lock(ChildStatsContextCS);
 
 	//Set up new StatsContext for this ThreadId
 	FNetStatsContext* NewChildStatsContext = new FNetStatsContext;
@@ -241,7 +241,7 @@ void FNetTypeStats::ReleaseChildNetStatsContext(FNetStatsContext* StatsContext)
 {
 	if(ensureMsgf(StatsContext, TEXT("ReleaseChildNetStatsContext was passed a null StatsContext. Possibility of net stats data being lost.")))
 	{
-		FScopeLock Lock(&ChildStatsContextCS);
+		UE::TScopeLock Lock(ChildStatsContextCS);
 		check(!StatsContext->CanAcquire());
 		StatsContext->Release();
 	}
diff --git a/Engine/Source/Runtime/Net/Iris/Private/Iris/Stats/NetStats.h b/Engine/Source/Runtime/Net/Iris/Private/Iris/Stats/NetStats.h
--- a/Engine/Source/Runtime/Net/Iris/Private/Iris/Stats/NetStats.h
+++ b/Engine/Source/Runtime/Net/Iris/Private/Iris/Stats/NetStats.h
@@ -6,6 +6,7 @@
 #include "HAL/CriticalSection.h"
 #include "Logging/LogMacros.h"
 #include "Containers/Array.h"
+#include "Misc/TransactionallySafeCriticalSection.h"
 
 DECLARE_LOG_CATEGORY_EXTERN(LogNetStats, Log, All);
 
@@ -94,7 +95,7 @@ private:
 		int32 ReplicatingConnectionCount = 0;
 	};
 
-	FCriticalSection CS;
+	FTransactionallySafeCriticalSection CS;
 	FStats Stats;
 };
 
@@ -232,7 +233,7 @@ private:
 	FNetStatsContext* ParentStatsContext = nullptr;
 
 	/** Critical Section used to prevent multiple threads from accessing ChildStatsContext array simultaneously */
-	FCriticalSection ChildStatsContextCS;
+	FTransactionallySafeCriticalSection ChildStatsContextCS;
 
 	/** NetStatsContexts that are used by sub tasks during a parallel phase. Are accumulated to ParentStatsContext at the end of a frame. */
 	TArray<FNetStatsContext*> ChildStatsContext;
