Iris - Implement Reset() to be called when reusing an FReplicatedObjectData instance. Much better code gen.

diff --git a/Engine/Source/Runtime/Net/Iris/Private/Iris/ReplicationSystem/NetRefHandleManager.cpp b/Engine/Source/Runtime/Net/Iris/Private/Iris/ReplicationSystem/NetRefHandleManager.cpp
--- a/Engine/Source/Runtime/Net/Iris/Private/Iris/ReplicationSystem/NetRefHandleManager.cpp
+++ b/Engine/Source/Runtime/Net/Iris/Private/Iris/ReplicationSystem/NetRefHandleManager.cpp
@@ -79,7 +79,7 @@ void FNetRefHandleManager::Init(const FInitParams& InitParams)
 	ReplicatedInstances = TNetChunkedArray<TObjectPtr<UObject>>(PreAllocatedNetChunkedArrayCount, EInitMemory::Zero);
 
 	// For convenience we initialize ReplicatedObjectData for InvalidInternalIndex so that GetReplicatedObjectDataNoCheck returns something useful.
-	ReplicatedObjectData[InvalidInternalIndex] = FReplicatedObjectData();
+	ReplicatedObjectData[InvalidInternalIndex].Reset();
 
 	// Init all NetBitArrays here
 	{
@@ -254,21 +254,18 @@ FInternalNetRefIndex FNetRefHandleManager::InternalCreateNetObject(const FNetRef
 
 	// Store data;
 	FReplicatedObjectData& Data = ReplicatedObjectData[InternalIndex];
-
-	Data = FReplicatedObjectData();
-
+	Data.Reset();
 	Data.RefHandle = NetRefHandle;
 	Data.NetHandle = GlobalHandle;
 	Data.Protocol = Params.ReplicationProtocol;
+	Data.bShouldPropagateChangedStates = 1U;
 	Data.NetFactoryId = Params.NetFactoryId;
 	Data.IrisAsyncLoadingPriority = Params.IrisAsyncLoadingPriority;
-	Data.InstanceProtocol = nullptr;
+
 	ObjectsWithPreUpdate.ClearBit(InternalIndex);
 	ObjectsWithFullPushBasedDirtiness.ClearBit(InternalIndex);
+	
 	ReplicatedObjectStateBuffers[InternalIndex] = nullptr;
-	Data.ReceiveStateBuffer = nullptr;
-	Data.bShouldPropagateChangedStates = 1U;
-	Data.bIsDebugObject = 0U;
 
 	++ActiveObjectCount;
 
@@ -503,7 +500,7 @@ void FNetRefHandleManager::InternalDestroyNetObject(FInternalNetRefIndex Interna
 
 	UE_NET_TRACE_NETHANDLE_DESTROYED(Data.RefHandle);
 
-	Data = FReplicatedObjectData();
+	Data.Reset();
 
 	// tracking
 	AssignedInternalIndices.ClearBit(InternalIndex);
diff --git a/Engine/Source/Runtime/Net/Iris/Private/Iris/ReplicationSystem/NetRefHandleManager.h b/Engine/Source/Runtime/Net/Iris/Private/Iris/ReplicationSystem/NetRefHandleManager.h
--- a/Engine/Source/Runtime/Net/Iris/Private/Iris/ReplicationSystem/NetRefHandleManager.h
+++ b/Engine/Source/Runtime/Net/Iris/Private/Iris/ReplicationSystem/NetRefHandleManager.h
@@ -88,9 +88,17 @@ public:
 	// $TODO: This should be split up into separate array according to usage patterns, it is getting a bit too large
 	struct FReplicatedObjectData
 	{
-		FReplicatedObjectData()
-		: Flags(0U)
+		FReplicatedObjectData() = default;
+
+		// Do not support moving/copying. We should modify entries in the array directly.
+		FReplicatedObjectData(FReplicatedObjectData&&) = delete;
+		FReplicatedObjectData& operator=(FReplicatedObjectData&&) = delete;
+		FReplicatedObjectData(const FReplicatedObjectData&) = delete;
+		FReplicatedObjectData& operator=(const FReplicatedObjectData&) = delete;
+
+		void Reset()
 		{
+			new (this) FReplicatedObjectData();
 		}
 
 		FNetRefHandle RefHandle;
@@ -98,7 +106,7 @@ public:
 		
 		const FReplicationProtocol* Protocol = nullptr;
 		const FReplicationInstanceProtocol* InstanceProtocol = nullptr;
-		uint8* ReceiveStateBuffer = nullptr ;
+		uint8* ReceiveStateBuffer = nullptr;
 		
 		/** Subobjects only: Internal index of the RootObject of this subobject */
 		FInternalNetRefIndex SubObjectRootIndex = InvalidInternalIndex;
@@ -106,31 +114,24 @@ public:
 		/** Subobjects only: Internal index of the ParentObject of this subobject */
 		FInternalNetRefIndex SubObjectParentIndex = InvalidInternalIndex;
 
+		uint16 bShouldPropagateChangedStates : 1U = 0;
+		uint16 bTearOff : 1U = 0;
+		uint16 bDestroySubObjectWithOwner : 1U = 0;
+		uint16 bIsDependentObject : 1U = 0;
+		uint16 bHasDependentObjects : 1U = 0;
+		uint16 bAllowDestroyInstanceFromRemote : 1U = 0;
+		uint16 bNeedsFullCopyAndQuantize : 1U = 0;
+		uint16 bWantsFullPoll : 1U = 0;
+		uint16 bPendingEndReplication : 1U = 0;
+		uint16 bHasCachedCreationInfo : 1U = 0;
+		uint16 bIsDebugObject : 1U = 0;
+
 		/** The factory responsible for instantiating this object */
 		FNetObjectFactoryId NetFactoryId = InvalidNetObjectFactoryId;
 
 		/** The async loading priority for references collected from this object */
 		EIrisAsyncLoadingPriority IrisAsyncLoadingPriority = EIrisAsyncLoadingPriority::Default;
 
-		union
-		{
-			uint16 Flags : 16U;
-			struct
-			{
-				uint16 bShouldPropagateChangedStates : 1U;
-				uint16 bTearOff : 1U;
-				uint16 bDestroySubObjectWithOwner : 1U;
-				uint16 bIsDependentObject : 1U;
-				uint16 bHasDependentObjects : 1U;
-				uint16 bAllowDestroyInstanceFromRemote : 1U;
-				uint16 bNeedsFullCopyAndQuantize : 1U;
-				uint16 bWantsFullPoll : 1U;
-				uint16 bPendingEndReplication : 1U;
-				uint16 bHasCachedCreationInfo : 1U;
-				uint16 bIsDebugObject : 1U;
-			};
-		};
-
 		ECloseReplicationReason CloseReason = ECloseReplicationReason::None;
 	
 		/** Returns true if this is a SubObject that gets destroyed when his RootObject is destroyed too. */
