Iris
- Exposed common cvar "net.Iris.ScheduleCreationDependenciesFirst" to test only relying on creation link dependencies for Entity replication. Disabled by default
- Added test that exercises the path that would break for huge deep hierarchies (super slow)

diff --git a/Engine/Plugins/Runtime/ReplicationSystemTestPlugin/Source/Private/Tests/ReplicationSystem/TestDependentObjects.cpp b/Engine/Plugins/Runtime/ReplicationSystemTestPlugin/Source/Private/Tests/ReplicationSystem/TestDependentObjects.cpp
--- a/Engine/Plugins/Runtime/ReplicationSystemTestPlugin/Source/Private/Tests/ReplicationSystem/TestDependentObjects.cpp
+++ b/Engine/Plugins/Runtime/ReplicationSystemTestPlugin/Source/Private/Tests/ReplicationSystem/TestDependentObjects.cpp
@@ -951,7 +951,7 @@ UE_NET_TEST_FIXTURE(FTestFilteredDependentObjectFixture, TestDependentObjectSche
 UE_NET_TEST_FIXTURE(FReplicationSystemServerClientTestFixture, TestCreationDependencyCreationOrder)
 {
 	//LOG_SCOPE_VERBOSITY_OVERRIDE(LogIris, ELogVerbosity::VeryVerbose);
-	TGuardConsoleVariable<bool> CVarOverride(TEXT("net.Iris.ReplicationWriter.ScheduleCreationDependenciesFirst"), true);
+	TGuardConsoleVariable<bool> CVarOverride(TEXT("net.Iris.ScheduleCreationDependenciesFirst"), true);
 
 	UReplicationSystem* ReplicationSystem = Server->ReplicationSystem;
 	UReplicatedTestObjectBridge* Bridge = Server->GetReplicationBridge();
@@ -1009,7 +1009,7 @@ UE_NET_TEST_FIXTURE(FReplicationSystemServerClientTestFixture, TestCreationDepen
 UE_NET_TEST_FIXTURE(FReplicationSystemServerClientTestFixture, TestNestedCreationDependenciesCreationOrder)
 {
 	//LOG_SCOPE_VERBOSITY_OVERRIDE(LogIris, ELogVerbosity::VeryVerbose);
-	TGuardConsoleVariable<bool> CVarOverride(TEXT("net.Iris.ReplicationWriter.ScheduleCreationDependenciesFirst"), true);
+	TGuardConsoleVariable<bool> CVarOverride(TEXT("net.Iris.ScheduleCreationDependenciesFirst"), true);
 
 	UReplicationSystem* ReplicationSystem = Server->ReplicationSystem;
 	UReplicatedTestObjectBridge* Bridge = Server->GetReplicationBridge();
diff --git a/Engine/Source/Runtime/Net/Iris/Private/Iris/IrisConfig.cpp b/Engine/Source/Runtime/Net/Iris/Private/Iris/IrisConfig.cpp
--- a/Engine/Source/Runtime/Net/Iris/Private/Iris/IrisConfig.cpp
+++ b/Engine/Source/Runtime/Net/Iris/Private/Iris/IrisConfig.cpp
@@ -9,17 +9,25 @@
 namespace UE::Net
 {
 
-static int32 CVarUseIrisReplication = 0;
-static FAutoConsoleVariableRef CVarUseIrisReplicationRef(TEXT("net.Iris.UseIrisReplication"), CVarUseIrisReplication, TEXT("Enables Iris replication system. 0 will fallback to legacy replicationsystem."), ECVF_Default );
+namespace Private::CVars
+{
+
+static int32 GUseIrisReplication = 0;
+static FAutoConsoleVariableRef CVarUseIrisReplicationRef(TEXT("net.Iris.UseIrisReplication"), GUseIrisReplication, TEXT("Enables Iris replication system. 0 will fallback to legacy replicationsystem."), ECVF_Default );
+
+static bool bScheduleCreationDependenciesFirst = false;
+static FAutoConsoleVariableRef CVarScheduleCreationDependenciesFirst(TEXT("net.Iris.ScheduleCreationDependenciesFirst"), bScheduleCreationDependenciesFirst, TEXT("If enabled we will schedule an objects creation dependencies before the object if they are not confirmed as created."));
+
+}
 
 bool ShouldUseIrisReplication()
 {
-	return CVarUseIrisReplication > 0;
+	return Private::CVars::GUseIrisReplication > 0;
 }
 
 void SetUseIrisReplication(bool EnableIrisReplication)
 {
-	CVarUseIrisReplication = EnableIrisReplication ? 1 : 0;
+	Private::CVars::GUseIrisReplication = EnableIrisReplication ? 1 : 0;
 }
 
 EReplicationSystem GetUseIrisReplicationCmdlineValue()
@@ -35,6 +43,11 @@ EReplicationSystem GetUseIrisReplicationCmdlineValue()
 	return EReplicationSystem::Default;
 }
 
+bool IsSchedulingCreationDependenciesFirst()
+{
+	return Private::CVars::bScheduleCreationDependenciesFirst;
+}
+
 }
 
 
diff --git a/Engine/Source/Runtime/Net/Iris/Private/Iris/ReplicationSystem/ReplicationWriter.cpp b/Engine/Source/Runtime/Net/Iris/Private/Iris/ReplicationSystem/ReplicationWriter.cpp
--- a/Engine/Source/Runtime/Net/Iris/Private/Iris/ReplicationSystem/ReplicationWriter.cpp
+++ b/Engine/Source/Runtime/Net/Iris/Private/Iris/ReplicationSystem/ReplicationWriter.cpp
@@ -111,10 +111,6 @@ static FAutoConsoleVariableRef CVarReplicationWriterCannotSendWarningInterval(TE
 static bool bAllowSendStateWithEndReplicationForSubObjects = false;
 static FAutoConsoleVariableRef CvarAllowSendStateWithEndReplicationForSubObjects(TEXT("net.Iris.AllowSendStateWithEndReplicationForSubObjects"), bAllowSendStateWithEndReplicationForSubObjects, TEXT("Allow state to be sent together with subobject destroy. Default is false. Cannot be true if net.Iris.ImmediateDispatchEndReplicationForSubObjects is true."));
 
-// net.Iris.ReplicationWriter.ScheduleCreationDependenciesFirst
-static bool bScheduleCreationDependenciesFirst = false;
-static FAutoConsoleVariableRef CVarScheduleCreationDependenciesFirst(TEXT("net.Iris.ReplicationWriter.ScheduleCreationDependenciesFirst"), bScheduleCreationDependenciesFirst, TEXT("If enabled we will schedule an objects creation dependencies before the object if they are not confirmed as created."));
-
 // Drop unsent ordered unreliable attachments at the end of the tick, 
 // -1 = Allow them to be queued for next tick, 
 // 0  = Always drop at the end of the tick
@@ -3090,8 +3086,6 @@ FReplicationWriter::EWriteObjectStatus FReplicationWriter::WriteObjectInBatch(FN
 			const bool bIsDependentInitialState = IsInitialState(GetReplicationInfo(DependentInternalIndex).GetState());
 			if (bIsDependentInitialState && !WriteContext.ObjectsWrittenThisPacket.GetBit(DependentInternalIndex))
 			{
-				UE_NET_TRACE_SCOPE(DependentObjectData, *Context.GetBitStreamWriter(), Context.GetTraceCollector(), ENetTraceVerbosity::VeryVerbose);
-
 				OutBatchInfo.ParentInternalIndex = DependentInternalIndex;
 				EWriteObjectStatus DependentObjectWriteStatus = WriteObjectInBatch(Context, DependentInternalIndex, WriteObjectTraits, OutBatchInfo);
 				if (!IsWriteObjectSuccess(DependentObjectWriteStatus))
@@ -3765,7 +3759,7 @@ uint32 FReplicationWriter::WriteObjects(FNetSerializationContext& Context)
 				checkSlow(InternalIndex != ObjectIndexForOOBAttachment);
 
 				// If the object has creation dependencies we will schedule unsent creation dependencies before we try to send the object again.
-				const bool bAppendCreationDependencies = bScheduleCreationDependenciesFirst;
+				const bool bAppendCreationDependencies = UE::Net::IsSchedulingCreationDependenciesFirst();
 				const EWriteStatus Status = SendObjectFunction(InternalIndex, bAppendCreationDependencies);
 				
 				if (Status == EWriteStatus::Abort)
diff --git a/Engine/Source/Runtime/Net/Iris/Public/Iris/IrisConfig.h b/Engine/Source/Runtime/Net/Iris/Public/Iris/IrisConfig.h
--- a/Engine/Source/Runtime/Net/Iris/Public/Iris/IrisConfig.h
+++ b/Engine/Source/Runtime/Net/Iris/Public/Iris/IrisConfig.h
@@ -19,6 +19,10 @@ IRISCORE_API void SetUseIrisReplication(bool EnableIrisReplication);
 /** Returns what replication sytem was set to be used by the cmdline. Returns Default when the command line was not set. */
 IRISCORE_API EReplicationSystem GetUseIrisReplicationCmdlineValue();
 
+/** Returns if we are scheduling creation dependencies before dependent objects */
+UE_EXPERIMENTAL(5.7, "Experimental and subject for removal.")
+IRISCORE_API bool IsSchedulingCreationDependenciesFirst();
+
 }
 
 /* NetBitStreamReader/Writer validation support */
