Iris
- Enabled net.Iris.ScheduleCreationDependenciesFirst by default.

Can be disabled with cvar "net.Iris.ScheduleCreationDependenciesFirst 0".

diff --git a/Engine/Source/Runtime/Net/Iris/Private/Iris/IrisConfig.cpp b/Engine/Source/Runtime/Net/Iris/Private/Iris/IrisConfig.cpp
--- a/Engine/Source/Runtime/Net/Iris/Private/Iris/IrisConfig.cpp
+++ b/Engine/Source/Runtime/Net/Iris/Private/Iris/IrisConfig.cpp
@@ -15,7 +15,7 @@ namespace Private::CVars
 static int32 GUseIrisReplication = 0;
 static FAutoConsoleVariableRef CVarUseIrisReplicationRef(TEXT("net.Iris.UseIrisReplication"), GUseIrisReplication, TEXT("Enables Iris replication system. 0 will fallback to legacy replicationsystem."), ECVF_Default );
 
-static bool bScheduleCreationDependenciesFirst = false;
+static bool bScheduleCreationDependenciesFirst = true;
 static FAutoConsoleVariableRef CVarScheduleCreationDependenciesFirst(TEXT("net.Iris.ScheduleCreationDependenciesFirst"), bScheduleCreationDependenciesFirst, TEXT("If enabled we will schedule an objects creation dependencies before the object if they are not confirmed as created."));
 
 }
