Iris Parallelization - Make Net Tracing existing objects and object creation threadsafe. Added a SharedNewObjects list to temporarily store newly created objects in parallel net connection tick and then merge them with the full list after NetConnection Tick has completed.

#tests Tested against Dedicated Server Cooked Build and ReplicationSystemTests
[RN] Net Tracing for object creation has been made thread safe for Iris Parallelization
#Synced-CL 47572861

diff --git a/Engine/Source/Runtime/Engine/Private/NetDriver.cpp b/Engine/Source/Runtime/Engine/Private/NetDriver.cpp
--- a/Engine/Source/Runtime/Engine/Private/NetDriver.cpp
+++ b/Engine/Source/Runtime/Engine/Private/NetDriver.cpp
@@ -1208,6 +1208,10 @@ void UNetDriver::TickFlush(float DeltaSeconds)
 
 		FlushHandler();
 	}
+	
+#if UE_NET_TRACE_ENABLED
+	FNetTrace::FinishConnectionTick(GetNetTraceId());
+#endif // UE_NET_TRACE_ENABLED
 
 	if (CVarNetDebugDraw.GetValueOnAnyThread() > 0)
 	{
@@ -7430,6 +7434,11 @@ void UNetDriver::SetReplicationSystem(UReplicationSystem* InReplicationSystem)
 	{
 		// When we run using Iris, we use ReplicationSystemId as our unique identifier
 		NetTraceId = ReplicationSystem->GetId();
+		
+#if UE_NET_TRACE_ENABLED
+		//We explicitly make sure NetTrace has a corresponding GameInstance set up, as it means we can assume the owning map is read-only during parallel phases of Iris
+		FNetTrace::SetupInstanceForNetTraceId(NetTraceId);
+#endif // UE_NET_TRACE_ENABLED
 
 		UE::Net::FReplicationSystemUtil::StartReplicatingActorsInWorldForNetDriver(GetWorld(), this);
 	}
diff --git a/Engine/Source/Runtime/Net/Core/Private/Net/Core/Trace/NetTraceInternal.cpp b/Engine/Source/Runtime/Net/Core/Private/Net/Core/Trace/NetTraceInternal.cpp
--- a/Engine/Source/Runtime/Net/Core/Private/Net/Core/Trace/NetTraceInternal.cpp
+++ b/Engine/Source/Runtime/Net/Core/Private/Net/Core/Trace/NetTraceInternal.cpp
@@ -4,6 +4,7 @@
 #include "Containers/ArrayView.h"
 #include "Containers/Map.h"
 #include "Containers/Ticker.h"
+#include "Misc/TransactionallySafeCriticalSection.h"
 #include "Net/Core/Trace/NetDebugName.h"
 #include "Net/Core/Trace/Reporters/NetTraceReporter.h"
 #include "Hash/CityHash.h"
@@ -53,7 +54,10 @@ struct FNetTraceInternal
 	{
 		TMap<uint32, FConnection> Connections;
 		TMap<uint64, FNetObject> Objects;
+		TSet<uint64> ObjectsTraced;
+		TMap<uint64, FNetObject> SharedNewObjects; //Thread-safe, only used to track objects that were added during this frame
 		FString Name;
+		uint32 Id;//Same as the index of the owning GameInstances TMap
 		bool bIsServer;
 	};
 
@@ -96,7 +100,37 @@ static std::atomic<UE::Net::FNetDebugNameId> NextNameId(1);
 static thread_local TUniquePtr<FNetTraceInternal::FThreadBuffer> ThreadBuffer;
 static FDelegateHandle OnTraceConnectedHandle;
 static FDelegateHandle OnTraceStoppedHandle;
-static TMap<uint32, FNetTraceInternal::FGameInstance> GameInstances;
+static FTransactionallySafeCriticalSection NewObjectsCS;
+
+// Container for GameInstances map which ensures each is set up with an Id, allowing
+// trace internals to pass FGameInstance around and avoid double map lookups
+class FGameInstanceContainer
+{
+public:
+	typedef TMap<uint32, FNetTraceInternal::FGameInstance> MapType;
+
+	FNetTraceInternal::FGameInstance& FindOrAdd(const uint32 GameInstanceId)
+	{
+		FNetTraceInternal::FGameInstance& GameInstance = GameInstances.FindOrAdd(GameInstanceId);
+		GameInstance.Id = GameInstanceId;
+		return GameInstance;
+	}
+
+	const MapType& GetGameInstances() const
+	{
+		return GameInstances;
+	}
+
+	void Remove(uint32 GameInstanceId)
+	{
+		GameInstances.Remove(GameInstanceId);
+	}
+
+private:
+	MapType GameInstances;
+};
+
+static FGameInstanceContainer GameInstanceContainer;
 
 FNetTrace::FOnResetPersistentNetDebugNames FNetTrace::OnResetPersistentNetDebugNames;
 FNetTrace::FOnNetTraceStarted FNetTrace::OnNetTraceStarted;
@@ -127,6 +161,11 @@ UE_AUTORTFM_ALWAYS_OPEN uint32 FNetTraceInternal::GetCurrentTraceId()
 	return CurrentTraceId;
 }
 
+void FNetTrace::SetupInstanceForNetTraceId(uint32 InNetTraceId)
+{
+	GameInstanceContainer.FindOrAdd(InNetTraceId);
+}
+
 void FNetTrace::OnTraceConnected()
 {
 	++CurrentTraceId;
@@ -136,10 +175,10 @@ void FNetTrace::OnTraceConnected()
 
 	FNetTraceInternal::Reporter::ReportInitEvent(FNetTraceInternal::NetTraceVersion);
 
-	for (TPair<uint32, FNetTraceInternal::FGameInstance>& InstanceIter : GameInstances)
+	for (const TPair<uint32, FNetTraceInternal::FGameInstance>& InstanceIter : GameInstanceContainer.GetGameInstances())
 	{
 		const uint32 GameInstanceId = InstanceIter.Key;
-		FNetTraceInternal::FGameInstance& GameInstance = InstanceIter.Value;
+		const FNetTraceInternal::FGameInstance& GameInstance = InstanceIter.Value;
 
 		FNetTraceInternal::Reporter::ReportInstanceUpdated(GameInstanceId, GameInstance.bIsServer, *GameInstance.Name);
 
@@ -207,7 +246,7 @@ void FNetTrace::SetTraceVerbosity(uint32 Verbosity)
 
 void FNetTrace::TraceEndSession(uint32 GameInstanceId)
 {
-	GameInstances.Remove(GameInstanceId);
+	GameInstanceContainer.Remove(GameInstanceId);
 
 	if (GNetTraceRuntimeVerbosity)
 	{
@@ -217,7 +256,7 @@ void FNetTrace::TraceEndSession(uint32 GameInstanceId)
 
 void FNetTrace::TraceInstanceUpdated(uint32 GameInstanceId, bool bIsServer, const TCHAR* Name)
 {
-	FNetTraceInternal::FGameInstance& Instance = GameInstances.FindOrAdd(GameInstanceId);
+	FNetTraceInternal::FGameInstance& Instance = GameInstanceContainer.FindOrAdd(GameInstanceId);
 	Instance.Name = Name;
 	Instance.bIsServer = bIsServer;
 
@@ -512,11 +551,25 @@ static UE::Net::FNetDebugNameId& GetOrCreateDebugNameIdInLocalMap(const UE::Net:
 	return *NameIdPtr;
 }
 
-static void TraceExistingObject(uint32 GameInstanceId, const TMap<uint64, FNetTraceInternal::FNetObject>& Objects, TSet<uint64>& ObjectsTraced, uint64 NetObjectId)
+static void TraceExistingObject(FNetTraceInternal::FGameInstance& GameInstance, const TMap<uint64, FNetTraceInternal::FNetObject>& Objects, TSet<uint64>& ObjectsTraced, uint64 NetObjectId)
 {
 	if (!ObjectsTraced.Contains(NetObjectId))
 	{
+		//Try to look up object in the Objects list, otherwise fall back to the SharedNewObjects List which can be accessed from any thread
 		const FNetTraceInternal::FNetObject* NetObject = Objects.Find(NetObjectId);
+		FNetTraceInternal::FNetObject CachedNetObject;//Conditionally used for data in SharedNewObjects to avoid locking NewObjectsCS for the whole function
+		if (!NetObject)
+		{
+			UE::TScopeLock CondScopeLock(NewObjectsCS);
+			const FNetTraceInternal::FNetObject* TempNetObjectPtr = GameInstance.SharedNewObjects.Find(NetObjectId);
+			if(TempNetObjectPtr)
+			{
+				//Copy data into CachedNetObject, in case another thread modifies SharedNewObjects and TempNetObjectPtr becomes invalid
+				CachedNetObject = *TempNetObjectPtr;
+				NetObject = &CachedNetObject;
+			}
+		}
+		
 		if (NetObject)
 		{
 			ObjectsTraced.Add(NetObjectId);
@@ -526,17 +579,17 @@ static void TraceExistingObject(uint32 GameInstanceId, const TMap<uint64, FNetTr
 				FNetTraceInternal::FThreadBuffer& LocalThreadBuffer = FNetTraceInternal::GetThreadBuffer();
 				const UE::Net::FNetDebugNameId DebugNameId = GetOrCreateDebugNameIdInLocalMap(NetObject->DebugName, LocalThreadBuffer);
 					
-				FNetTraceInternal::Reporter::ReportObjectExists(GameInstanceId, NetObjectId, DebugNameId, NetObject->TypeIdentifier, NetObject->OwnerId);
+				FNetTraceInternal::Reporter::ReportObjectExists(GameInstance.Id, NetObjectId, DebugNameId, NetObject->TypeIdentifier, NetObject->OwnerId);
 			}
 			else
 			{
-				FNetTraceInternal::Reporter::ReportObjectExists(GameInstanceId, NetObjectId, FNetTrace::TraceName(NetObject->Name), NetObject->TypeIdentifier, NetObject->OwnerId);
+				FNetTraceInternal::Reporter::ReportObjectExists(GameInstance.Id, NetObjectId, FNetTrace::TraceName(NetObject->Name), NetObject->TypeIdentifier, NetObject->OwnerId);
 			}
 
 			if (NetObject->ParentNetObjectId != uint64(-1))
 			{
-				TraceExistingObject(GameInstanceId, Objects, ObjectsTraced, NetObject->ParentNetObjectId);
-				FNetTraceInternal::Reporter::ReportSubObject(GameInstanceId, NetObject->ParentNetObjectId, NetObjectId);
+				TraceExistingObject(GameInstance, Objects, ObjectsTraced, NetObject->ParentNetObjectId);
+				FNetTraceInternal::Reporter::ReportSubObject(GameInstance.Id, NetObject->ParentNetObjectId, NetObjectId);
 			}
 		}
 	}
@@ -550,12 +603,12 @@ void FNetTrace::TraceCollectedEvents(FNetTraceCollector& Collector, uint32 GameI
 	PacketInfo.PacketSequenceNumber = 0;
 	PacketInfo.PacketType = PacketType;
 
-	FNetTraceInternal::FGameInstance& GameInstance = GameInstances.FindOrAdd(GameInstanceId);
+	FNetTraceInternal::FGameInstance& GameInstance = GameInstanceContainer.FindOrAdd(GameInstanceId);
 	FNetTraceInternal::FThreadBuffer& LocalThreadBuffer = FNetTraceInternal::GetThreadBuffer();
 	FNetTraceInternal::FObjectTraceCounter& ObjectTraceData = LocalThreadBuffer.GameInstanceToObjectTraceCounterMap.FindOrAdd(GameInstanceId);
 	for (uint64 NetObjectId : Collector.ReferencedObjectIds)
 	{
-		TraceExistingObject(GameInstanceId, GameInstance.Objects, ObjectTraceData.ObjectsTraced, NetObjectId);
+		TraceExistingObject(GameInstance, GameInstance.Objects, ObjectTraceData.ObjectsTraced, NetObjectId);
 	}
 
 	// Trace all collected events
@@ -596,18 +649,20 @@ void FNetTrace::TracePacket(uint32 GameInstanceId, uint32 ConnectionId, uint32 P
 
 void FNetTrace::TraceObjectCreated(uint32 GameInstanceId, uint64 NetObjectId, const UE::Net::FNetDebugName* DebugName, uint64 TypeIdentifier, uint32 OwnerId)
 {
-	FNetTraceInternal::FGameInstance& GameInstance = GameInstances.FindOrAdd(GameInstanceId);
-	FNetTraceInternal::FThreadBuffer& LocalThreadBuffer = FNetTraceInternal::GetThreadBuffer();
-	FNetTraceInternal::FObjectTraceCounter& ObjectTraceData = LocalThreadBuffer.GameInstanceToObjectTraceCounterMap.FindOrAdd(GameInstanceId);
+	FNetTraceInternal::FGameInstance& GameInstance = GameInstanceContainer.FindOrAdd(GameInstanceId);
 
-	GameInstance.Objects.Emplace(NetObjectId, { .Name = NAME_None, .DebugName = DebugName, .TypeIdentifier = TypeIdentifier, .ParentNetObjectId = uint64(-1), .OwnerId = OwnerId });
+	{
+		UE::TScopeLock NewObjectsScopeLock(NewObjectsCS);
+		GameInstance.SharedNewObjects.Emplace(NetObjectId, { .Name = NAME_None, .DebugName = DebugName, .TypeIdentifier = TypeIdentifier, .ParentNetObjectId = uint64(-1), .OwnerId = OwnerId });
+	}
 
 	if (!GNetTraceRuntimeVerbosity)
 	{
 		return;
 	}
 
-	ensureMsgf(IsInGameThread(), TEXT("TraceObjectCreated running on parallel thread - this is not yet supported"));
+	FNetTraceInternal::FThreadBuffer& LocalThreadBuffer = FNetTraceInternal::GetThreadBuffer();
+	FNetTraceInternal::FObjectTraceCounter& ObjectTraceData = LocalThreadBuffer.GameInstanceToObjectTraceCounterMap.FindOrAdd(GameInstanceId);
 
 	ObjectTraceData.ObjectsTraced.Add(NetObjectId);
 
@@ -620,17 +675,20 @@ void FNetTrace::TraceObjectCreated(uint32 GameInstanceId, uint64 NetObjectId, co
 
 void FNetTrace::TraceObjectCreated(uint32 GameInstanceId, uint64 NetObjectId, const FName ObjectName, uint64 TypeIdentifier, uint32 OwnerId)
 {
-	FNetTraceInternal::FGameInstance& GameInstance = GameInstances.FindOrAdd(GameInstanceId);
-	FNetTraceInternal::FThreadBuffer& LocalThreadBuffer = FNetTraceInternal::GetThreadBuffer();
-	FNetTraceInternal::FObjectTraceCounter& ObjectTraceData = LocalThreadBuffer.GameInstanceToObjectTraceCounterMap.FindOrAdd(GameInstanceId);
-
-	GameInstance.Objects.Emplace(NetObjectId, { .Name = ObjectName, .DebugName = nullptr, .TypeIdentifier = TypeIdentifier, .ParentNetObjectId = uint64(-1), .OwnerId = OwnerId });
+	FNetTraceInternal::FGameInstance& GameInstance = GameInstanceContainer.FindOrAdd(GameInstanceId);
+	{
+		UE::TScopeLock NewObjectsScopeLock(NewObjectsCS);
+		GameInstance.SharedNewObjects.Emplace(NetObjectId, { .Name = ObjectName, .DebugName = nullptr, .TypeIdentifier = TypeIdentifier, .ParentNetObjectId = uint64(-1), .OwnerId = OwnerId });
+	}
 	
 	if (!GNetTraceRuntimeVerbosity)
 	{
 		return;
 	}
 
+	FNetTraceInternal::FThreadBuffer& LocalThreadBuffer = FNetTraceInternal::GetThreadBuffer();
+	FNetTraceInternal::FObjectTraceCounter& ObjectTraceData = LocalThreadBuffer.GameInstanceToObjectTraceCounterMap.FindOrAdd(GameInstanceId);
+
 	ObjectTraceData.ObjectsTraced.Add(NetObjectId);
 
 	FNetTraceInternal::Reporter::ReportObjectCreated(GameInstanceId, NetObjectId, FNetTrace::TraceName(ObjectName), TypeIdentifier, OwnerId);
@@ -638,7 +696,9 @@ void FNetTrace::TraceObjectCreated(uint32 GameInstanceId, uint64 NetObjectId, co
 
 void FNetTrace::TraceObjectDestroyed(uint32 GameInstanceId, uint64 NetObjectId)
 {
-	FNetTraceInternal::FGameInstance& GameInstance = GameInstances.FindOrAdd(GameInstanceId);
+	ensureMsgf(IsInGameThread(), TEXT("TraceObjectDestroyed must be called on the GameThread after parallel tasks have joined"));
+
+	FNetTraceInternal::FGameInstance& GameInstance = GameInstanceContainer.FindOrAdd(GameInstanceId);
 	FNetTraceInternal::FThreadBuffer& LocalThreadBuffer = FNetTraceInternal::GetThreadBuffer();
 	FNetTraceInternal::FObjectTraceCounter& ObjectTraceData = LocalThreadBuffer.GameInstanceToObjectTraceCounterMap.FindOrAdd(GameInstanceId);
 
@@ -655,7 +715,7 @@ void FNetTrace::TraceObjectDestroyed(uint32 GameInstanceId, uint64 NetObjectId)
 void FNetTrace::TraceConnectionCreated(uint32 GameInstanceId, uint32 ConnectionId)
 {
 	// TraceConnectionStateUpdated can be called before this, so use FindOrAdd
-	GameInstances.FindOrAdd(GameInstanceId).Connections.FindOrAdd(ConnectionId);
+	GameInstanceContainer.FindOrAdd(GameInstanceId).Connections.FindOrAdd(ConnectionId);
 
 	if (GNetTraceRuntimeVerbosity)
 	{
@@ -665,7 +725,7 @@ void FNetTrace::TraceConnectionCreated(uint32 GameInstanceId, uint32 ConnectionI
 
 void FNetTrace::TraceConnectionStateUpdated(uint32 GameInstanceId, uint32 ConnectionId, uint8 ConnectionStateValue)
 {
-	GameInstances.FindOrAdd(GameInstanceId).Connections.FindOrAdd(ConnectionId).State = ConnectionStateValue;
+	GameInstanceContainer.FindOrAdd(GameInstanceId).Connections.FindOrAdd(ConnectionId).State = ConnectionStateValue;
 
 	if (GNetTraceRuntimeVerbosity)
 	{
@@ -675,7 +735,7 @@ void FNetTrace::TraceConnectionStateUpdated(uint32 GameInstanceId, uint32 Connec
 
 void FNetTrace::TraceConnectionUpdated(uint32 GameInstanceId, uint32 ConnectionId, const TCHAR* AddressString, const TCHAR* OwningActor)
 {
-	FNetTraceInternal::FConnection& Connection = GameInstances.FindOrAdd(GameInstanceId).Connections.FindOrAdd(ConnectionId);
+	FNetTraceInternal::FConnection& Connection = GameInstanceContainer.FindOrAdd(GameInstanceId).Connections.FindOrAdd(ConnectionId);
 	Connection.Address = AddressString;
 	Connection.OwningActor = OwningActor;
 
@@ -806,11 +866,11 @@ UE::Net::FNetDebugNameId FNetTrace::TraceName(const UE::Net::FNetDebugName* Debu
 void FNetTrace::TraceObjectPolled(uint32 GameInstanceId, uint64 NetObjectId, uint64 Cycles, bool bIsWaste)
 {
 	// not rechecking GNetTraceRuntimeVerbosity here as it's already checked to be == VeryVerbose to have got this far
-	FNetTraceInternal::FGameInstance& GameInstance = GameInstances.FindOrAdd(GameInstanceId);
+	FNetTraceInternal::FGameInstance& GameInstance = GameInstanceContainer.FindOrAdd(GameInstanceId);
 	FNetTraceInternal::FThreadBuffer& LocalThreadBuffer = FNetTraceInternal::GetThreadBuffer();
 	FNetTraceInternal::FObjectTraceCounter& ObjectTraceData = LocalThreadBuffer.GameInstanceToObjectTraceCounterMap.FindOrAdd(GameInstanceId);
 	
-	TraceExistingObject(GameInstanceId, GameInstance.Objects, ObjectTraceData.ObjectsTraced, NetObjectId);
+	TraceExistingObject(GameInstance, GameInstance.Objects, ObjectTraceData.ObjectsTraced, NetObjectId);
 	
 	FNetTraceInternal::Reporter::ReportObjectPolled(GameInstanceId, NetObjectId, Cycles, bIsWaste);
 }
@@ -818,11 +878,11 @@ void FNetTrace::TraceObjectPolled(uint32 GameInstanceId, uint64 NetObjectId, uin
 void FNetTrace::TraceObjectQuantized(uint32 GameInstanceId, uint64 NetObjectId, uint64 Cycles)
 {
 	// not rechecking GNetTraceRuntimeVerbosity here as it's already checked to be == VeryVerbose to have got this far
-	FNetTraceInternal::FGameInstance& GameInstance = GameInstances.FindOrAdd(GameInstanceId);
+	FNetTraceInternal::FGameInstance& GameInstance = GameInstanceContainer.FindOrAdd(GameInstanceId);
 	FNetTraceInternal::FThreadBuffer& LocalThreadBuffer = FNetTraceInternal::GetThreadBuffer();
 	FNetTraceInternal::FObjectTraceCounter& ObjectTraceData = LocalThreadBuffer.GameInstanceToObjectTraceCounterMap.FindOrAdd(GameInstanceId);
 
-	TraceExistingObject(GameInstanceId, GameInstance.Objects, ObjectTraceData.ObjectsTraced, NetObjectId);
+	TraceExistingObject(GameInstance, GameInstance.Objects, ObjectTraceData.ObjectsTraced, NetObjectId);
 
 	FNetTraceInternal::Reporter::ReportObjectQuantized(GameInstanceId, NetObjectId, Cycles);
 }
@@ -830,18 +890,18 @@ void FNetTrace::TraceObjectQuantized(uint32 GameInstanceId, uint64 NetObjectId,
 void FNetTrace::TraceObjectWritten(uint32 GameInstanceId, uint64 NetObjectId, uint64 Cycles)
 {
 	// not rechecking GNetTraceRuntimeVerbosity here as it's already checked to be == VeryVerbose to have got this far
-	FNetTraceInternal::FGameInstance& GameInstance = GameInstances.FindOrAdd(GameInstanceId);
+	FNetTraceInternal::FGameInstance& GameInstance = GameInstanceContainer.FindOrAdd(GameInstanceId);
 	FNetTraceInternal::FThreadBuffer& LocalThreadBuffer = FNetTraceInternal::GetThreadBuffer();
 	FNetTraceInternal::FObjectTraceCounter& ObjectTraceData = LocalThreadBuffer.GameInstanceToObjectTraceCounterMap.FindOrAdd(GameInstanceId);
 
-	TraceExistingObject(GameInstanceId, GameInstance.Objects, ObjectTraceData.ObjectsTraced, NetObjectId);
+	TraceExistingObject(GameInstance, GameInstance.Objects, ObjectTraceData.ObjectsTraced, NetObjectId);
 
 	FNetTraceInternal::Reporter::ReportObjectWritten(GameInstanceId, NetObjectId, Cycles);
 }
 
 void FNetTrace::TraceSubObject(uint32 GameInstanceId, uint64 ParentNetObjectId, uint64 SubObjectNetObjectId)
 {
-	FNetTraceInternal::FGameInstance& GameInstance = GameInstances.FindOrAdd(GameInstanceId);
+	FNetTraceInternal::FGameInstance& GameInstance = GameInstanceContainer.FindOrAdd(GameInstanceId);
 	FNetTraceInternal::FThreadBuffer& LocalThreadBuffer = FNetTraceInternal::GetThreadBuffer();
 	if (FNetTraceInternal::FNetObject* SubObject = GameInstance.Objects.Find(SubObjectNetObjectId))
 	{
@@ -854,6 +914,24 @@ void FNetTrace::TraceSubObject(uint32 GameInstanceId, uint64 ParentNetObjectId,
 	}	
 }
 
+void FNetTrace::FinishConnectionTick(const uint32 GameInstanceId)
+{
+	ensureMsgf(IsInGameThread(), TEXT("FinishConnectionTick must be called on the GameThread after Connection->Tick is complete and parallel tasks have joined"));
+
+	FNetTraceInternal::FGameInstance& GameInstance = GameInstanceContainer.FindOrAdd(GameInstanceId);
+
+	//Combine the SharedNewObjects into the Objects list, and wipe it, ready for the next frame
+	for (const TPair<uint64, FNetTraceInternal::FNetObject>& KVP : GameInstance.SharedNewObjects)
+	{
+		if (!GameInstance.Objects.Find(KVP.Key))
+		{
+			GameInstance.Objects.Add(KVP.Key, { .Name = KVP.Value.Name, .DebugName = KVP.Value.DebugName, .TypeIdentifier = KVP.Value.TypeIdentifier, .ParentNetObjectId = KVP.Value.ParentNetObjectId, .OwnerId = KVP.Value.OwnerId });
+		}
+	}
+
+	GameInstance.SharedNewObjects.Empty();
+}
+
 static FAutoConsoleCommand NeTraceSetVerbosityCmd = FAutoConsoleCommand(
 	TEXT("NetTrace.SetTraceVerbosity"),
 	TEXT("Start NetTrace with given verbositylevel."),
diff --git a/Engine/Source/Runtime/Net/Core/Private/Net/Core/Trace/Reporters/NetTraceReporter.cpp b/Engine/Source/Runtime/Net/Core/Private/Net/Core/Trace/Reporters/NetTraceReporter.cpp
--- a/Engine/Source/Runtime/Net/Core/Private/Net/Core/Trace/Reporters/NetTraceReporter.cpp
+++ b/Engine/Source/Runtime/Net/Core/Private/Net/Core/Trace/Reporters/NetTraceReporter.cpp
@@ -29,7 +29,7 @@ UE_TRACE_EVENT_BEGIN(NetTrace, NameEvent, NoSync)
 	UE_TRACE_EVENT_FIELD(UE::Trace::AnsiString, Name)
 UE_TRACE_EVENT_END()
 
-UE_TRACE_EVENT_BEGIN(NetTrace, ObjectCreatedEvent)
+UE_TRACE_EVENT_BEGIN(NetTrace, ObjectCreatedEvent, NoSync)
 	UE_TRACE_EVENT_FIELD(uint64, TypeId)
 	UE_TRACE_EVENT_FIELD(uint64, ObjectId)
 	UE_TRACE_EVENT_FIELD(uint32, OwnerId)
@@ -77,14 +77,14 @@ UE_TRACE_EVENT_BEGIN(NetTrace, ConnectionClosedEvent)
 UE_TRACE_EVENT_END()
 
 // StatsCounterEvent
-UE_TRACE_EVENT_BEGIN(NetTrace, PacketStatsCounterEvent)
+UE_TRACE_EVENT_BEGIN(NetTrace, PacketStatsCounterEvent, NoSync)
 	UE_TRACE_EVENT_FIELD(uint32, StatsValue)
 	UE_TRACE_EVENT_FIELD(uint32, NameId)
 	UE_TRACE_EVENT_FIELD(uint16, ConnectionId)
 	UE_TRACE_EVENT_FIELD(uint8, GameInstanceId)
 UE_TRACE_EVENT_END()
 
-UE_TRACE_EVENT_BEGIN(NetTrace, FrameStatsCounterEvent)
+UE_TRACE_EVENT_BEGIN(NetTrace, FrameStatsCounterEvent, NoSync)
 	UE_TRACE_EVENT_FIELD(uint64, Timestamp)
 	UE_TRACE_EVENT_FIELD(uint32, StatsValue)
 	UE_TRACE_EVENT_FIELD(uint32, NameId)
@@ -104,7 +104,7 @@ UE_TRACE_EVENT_BEGIN(NetTrace, InstanceDestroyedEvent)
 UE_TRACE_EVENT_END()
 
 // Packet data is transmitted as attachment
-UE_TRACE_EVENT_BEGIN(NetTrace, PacketContentEvent)
+UE_TRACE_EVENT_BEGIN(NetTrace, PacketContentEvent, NoSync)
 	UE_TRACE_EVENT_FIELD(uint16, ConnectionId)
 	UE_TRACE_EVENT_FIELD(uint8, GameInstanceId)
 	UE_TRACE_EVENT_FIELD(uint8, PacketType)
@@ -142,7 +142,7 @@ UE_TRACE_EVENT_BEGIN(NetTrace, ObjectQuantizedEvent)
 	UE_TRACE_EVENT_FIELD(uint8, GameInstanceId)
 UE_TRACE_EVENT_END()
 
-UE_TRACE_EVENT_BEGIN(NetTrace, ObjectWrittenEvent)
+UE_TRACE_EVENT_BEGIN(NetTrace, ObjectWrittenEvent, NoSync)
 	UE_TRACE_EVENT_FIELD(uint64, ObjectId)
 	UE_TRACE_EVENT_FIELD(uint64, Cycles)
 	UE_TRACE_EVENT_FIELD(uint8, GameInstanceId)
diff --git a/Engine/Source/Runtime/Net/Core/Public/Net/Core/Trace/Private/NetTraceInternal.h b/Engine/Source/Runtime/Net/Core/Public/Net/Core/Trace/Private/NetTraceInternal.h
--- a/Engine/Source/Runtime/Net/Core/Public/Net/Core/Trace/Private/NetTraceInternal.h
+++ b/Engine/Source/Runtime/Net/Core/Public/Net/Core/Trace/Private/NetTraceInternal.h
@@ -309,6 +309,10 @@ struct FNetTrace
 
 	NETCORE_API static void TraceSubObject(uint32 GameInstanceId, uint64 ParentNetObjectId, uint64 SubObjectNetObjectId);
 
+	NETCORE_API static void SetupInstanceForNetTraceId(uint32 InNetTraceId);
+
+	NETCORE_API static void FinishConnectionTick(const uint32 GameInstanceId);
+
 	static bool constexpr GetNetTraceVerbosityEnabled(const uint32 V) { return (uint32)UE_NET_TRACE_COMPILETIME_VERBOSITY >= V; }
 
 	inline static FNetTraceCollector* GetCollectorAtVerbosity(FNetTraceCollector* Collector, uint32 CollectorVerbosity);
