Iris - remove unnecessary ensure in USphereNetObjectPrioritizer.

#Synced-CL Unknown

diff --git a/Engine/Source/Runtime/Net/Iris/Private/Iris/ReplicationSystem/Prioritization/SphereNetObjectPrioritizer.cpp b/Engine/Source/Runtime/Net/Iris/Private/Iris/ReplicationSystem/Prioritization/SphereNetObjectPrioritizer.cpp
--- a/Engine/Source/Runtime/Net/Iris/Private/Iris/ReplicationSystem/Prioritization/SphereNetObjectPrioritizer.cpp
+++ b/Engine/Source/Runtime/Net/Iris/Private/Iris/ReplicationSystem/Prioritization/SphereNetObjectPrioritizer.cpp
@@ -243,7 +243,6 @@ void USphereNetObjectPrioritizer::PrioritizeBatchForMultiView(FBatchParams& Batc
 		const FVector& ViewPosVector = View.Pos;
 		ViewPositions.Add(VectorLoadFloat3_W0(&ViewPosVector));
 	}
-	ensureMsgf(ViewPositions.Num() <= 8, TEXT("Performance warning: Global allocation was needed to accommodate %d views."), ViewPositions.Num());
 
 	const VectorRegister MaxFloatVector = VectorSetFloat1(MAX_flt);
 
