Iris - Implemented GetRootObjectIndexOfAnyObject in FNetRefHandleManager, returning either the passed parameter or, in the case of a subobject, its root object index.

#Synced-CL 47893047

diff --git a/Engine/Source/Runtime/Net/Iris/Private/Iris/ReplicationSystem/NetRefHandleManager.cpp b/Engine/Source/Runtime/Net/Iris/Private/Iris/ReplicationSystem/NetRefHandleManager.cpp
--- a/Engine/Source/Runtime/Net/Iris/Private/Iris/ReplicationSystem/NetRefHandleManager.cpp
+++ b/Engine/Source/Runtime/Net/Iris/Private/Iris/ReplicationSystem/NetRefHandleManager.cpp
@@ -839,9 +839,17 @@ FNetRefHandle FNetRefHandleManager::GetRootObjectOfSubObject(FNetRefHandle SubOb
 	return IsValidIndex(OwnerInternalIndex) ? ReplicatedObjectData[OwnerInternalIndex].RefHandle : FNetRefHandle();
 }
 
+FInternalNetRefIndex FNetRefHandleManager::GetRootObjectInternalIndexOfAnyObject(FInternalNetRefIndex InternalIndex) const
+{
+	// Find the rootobject for subojects, otherwise just use the passed index since it is a root or invalid.
+	const FReplicatedObjectData& ObjectData = ReplicatedObjectData[InternalIndex];
+	return ObjectData.IsSubObject() ? ObjectData.SubObjectRootIndex : InternalIndex;
+}
+
 FInternalNetRefIndex FNetRefHandleManager::GetRootObjectInternalIndexOfSubObject(FInternalNetRefIndex SubObjectIndex) const
 {
-	return SubObjectIndex != InvalidInternalIndex ? ReplicatedObjectData[SubObjectIndex].SubObjectRootIndex : InvalidInternalIndex;
+	// If InvalidInternalIndex is passed we can rely on that data being default initialized so it will return InvalidInternalIndex too.
+	return ReplicatedObjectData[SubObjectIndex].SubObjectRootIndex;
 }
 
 bool FNetRefHandleManager::AddDependentObject(FNetRefHandle ParentRefHandle, FNetRefHandle DependentObjectRefHandle, EDependentObjectSchedulingHint SchedulingHint)
diff --git a/Engine/Source/Runtime/Net/Iris/Private/Iris/ReplicationSystem/NetRefHandleManager.h b/Engine/Source/Runtime/Net/Iris/Private/Iris/ReplicationSystem/NetRefHandleManager.h
--- a/Engine/Source/Runtime/Net/Iris/Private/Iris/ReplicationSystem/NetRefHandleManager.h
+++ b/Engine/Source/Runtime/Net/Iris/Private/Iris/ReplicationSystem/NetRefHandleManager.h
@@ -288,6 +288,11 @@ public:
 	
 	FNetRefHandle GetRootObjectOfAnyObject(FNetRefHandle NetRefHandle) const;
 	FNetRefHandle GetRootObjectOfSubObject(FNetRefHandle SubObjectHandle) const;
+
+	// If InternalIndex is a root object or invalid it will be returned. If it's a subobject then its root object index will be returned.
+	FInternalNetRefIndex GetRootObjectInternalIndexOfAnyObject(FInternalNetRefIndex InternalIndex) const;
+
+	// If SubObjectIndex is valid and a subobject then its root object index is returned, otherwise InvalidInternalIndex is returned.
 	FInternalNetRefIndex GetRootObjectInternalIndexOfSubObject(FInternalNetRefIndex SubObjectIndex) const;
 
 	bool SetSubObjectNetCondition(FInternalNetRefIndex SubObjectInternalIndex, FLifeTimeConditionStorage SubObjectCondition, bool& bOutWasModified);
