Iris
- Added some additional validation in ReplicationWriter::StartReplication to trap crash

#Synced-CL 48086225

diff --git a/Engine/Plugins/Runtime/ReplicationSystemTestPlugin/Source/Private/Tests/Core/TestParallelTask.cpp b/Engine/Plugins/Runtime/ReplicationSystemTestPlugin/Source/Private/Tests/Core/TestParallelTask.cpp
--- a/Engine/Plugins/Runtime/ReplicationSystemTestPlugin/Source/Private/Tests/Core/TestParallelTask.cpp
+++ b/Engine/Plugins/Runtime/ReplicationSystemTestPlugin/Source/Private/Tests/Core/TestParallelTask.cpp
@@ -39,7 +39,7 @@ UE_NET_TEST_FIXTURE(FTestParallelTaskFixture, PollTask)
 	Server->NetUpdate();
 
 	const FNetRefHandleManager& NetRefHandleManager = Server->GetReplicationSystem()->GetReplicationSystemInternal()->GetNetRefHandleManager();
-	FNetBitArrayView ObjectsConsideredForPolling = NetRefHandleManager.GetAssignedInternalIndices();
+	FNetBitArrayView ObjectsConsideredForPolling = NetRefHandleManager.GetAssignedInternalIndicesView();
 	ensure(ObjectsConsideredForPolling.IsAnyBitSet());
 
 
diff --git a/Engine/Source/Runtime/Net/Iris/Private/Iris/Core/IrisDebugging.cpp b/Engine/Source/Runtime/Net/Iris/Private/Iris/Core/IrisDebugging.cpp
--- a/Engine/Source/Runtime/Net/Iris/Private/Iris/Core/IrisDebugging.cpp
+++ b/Engine/Source/Runtime/Net/Iris/Private/Iris/Core/IrisDebugging.cpp
@@ -520,7 +520,7 @@ void NetObjectProtocolReferencesToString(FStringBuilderBase& StringBuilder, FRep
 				}
 			};
 			
-			NetRefHandleManager.GetAssignedInternalIndices().ForAllSetBits(FindMatchingProtocol);
+			NetRefHandleManager.GetAssignedInternalIndicesView().ForAllSetBits(FindMatchingProtocol);
 		}
 	);
 }
diff --git a/Engine/Source/Runtime/Net/Iris/Private/Iris/ReplicationSystem/NetRefHandleManager.h b/Engine/Source/Runtime/Net/Iris/Private/Iris/ReplicationSystem/NetRefHandleManager.h
--- a/Engine/Source/Runtime/Net/Iris/Private/Iris/ReplicationSystem/NetRefHandleManager.h
+++ b/Engine/Source/Runtime/Net/Iris/Private/Iris/ReplicationSystem/NetRefHandleManager.h
@@ -276,7 +276,8 @@ public:
 	FNetBitArrayView GetDirtyObjectsToQuantize() const { return MakeNetBitArrayView(DirtyObjectsToQuantize); }
 
 	// Get bitarray for all internal indices that currently are assigned
-	const FNetBitArrayView GetAssignedInternalIndices() const { return MakeNetBitArrayView(AssignedInternalIndices); }
+	const FNetBitArray& GetAssignedInternalIndices() const { return AssignedInternalIndices; }
+	const FNetBitArrayView GetAssignedInternalIndicesView() const { return MakeNetBitArrayView(AssignedInternalIndices); }
 
 	// SubObjects
 	const FNetBitArray& GetSubObjectInternalIndices() const { return SubObjectInternalIndices; }
diff --git a/Engine/Source/Runtime/Net/Iris/Private/Iris/ReplicationSystem/ObjectReplicationBridge.cpp b/Engine/Source/Runtime/Net/Iris/Private/Iris/ReplicationSystem/ObjectReplicationBridge.cpp
--- a/Engine/Source/Runtime/Net/Iris/Private/Iris/ReplicationSystem/ObjectReplicationBridge.cpp
+++ b/Engine/Source/Runtime/Net/Iris/Private/Iris/ReplicationSystem/ObjectReplicationBridge.cpp
@@ -1656,7 +1656,7 @@ void UObjectReplicationBridge::PruneStaleObjects()
 	};
 
 	// Iterate over assigned indices and detect if any of the replicated instances has been garbagecollected (excluding DestroyedStartupObjectInternalIndices) as they never have an instance
-	FNetBitArrayView::ForAllSetBits(LocalNetRefHandleManager.GetAssignedInternalIndices(), LocalNetRefHandleManager.GetDestroyedStartupObjectInternalIndices(), FNetBitArrayView::AndNotOp, DetectStaleObjectsFunc);
+	FNetBitArrayView::ForAllSetBits(LocalNetRefHandleManager.GetAssignedInternalIndicesView(), LocalNetRefHandleManager.GetDestroyedStartupObjectInternalIndices(), FNetBitArrayView::AndNotOp, DetectStaleObjectsFunc);
 
 	// EndReplication/detach stale instances
 	for (FNetRefHandle Handle : MakeArrayView(StaleObjects.GetData(), StaleObjects.Num()))
diff --git a/Engine/Source/Runtime/Net/Iris/Private/Iris/ReplicationSystem/ReplicationWriter.cpp b/Engine/Source/Runtime/Net/Iris/Private/Iris/ReplicationSystem/ReplicationWriter.cpp
--- a/Engine/Source/Runtime/Net/Iris/Private/Iris/ReplicationSystem/ReplicationWriter.cpp
+++ b/Engine/Source/Runtime/Net/Iris/Private/Iris/ReplicationSystem/ReplicationWriter.cpp
@@ -511,8 +511,19 @@ void FReplicationWriter::StartReplication(uint32 InternalIndex)
 {
 	FReplicationInfo& Info = GetReplicationInfo(InternalIndex);
 
-	ensureMsgf(Info.GetState() == EReplicatedObjectState::Invalid, TEXT("Object ( InternalIndex: %u ) is in state %s in StartReplication."), InternalIndex, LexToString(Info.GetState()));
-	if (InternalIndex != ObjectIndexForOOBAttachment && Attachments.HasUnsentAttachments(ENetObjectAttachmentType::Normal, InternalIndex))
+	if (!ensureMsgf(Info.GetState() == EReplicatedObjectState::Invalid, TEXT("Object ( InternalIndex: %u ) is in state %s in StartReplication."), InternalIndex, LexToString(Info.GetState())))
+	{
+		return;
+	}
+
+	if (InternalIndex == ObjectIndexForOOBAttachment || !NetRefHandleManager->GetAssignedInternalIndices().GetBit(InternalIndex))
+	{
+		UE_LOG(LogIris, Error, TEXT("FReplicationWriter::StartReplication - Unexpected call to StartReplication for not assigned index ( InternalIndex: %u )"), InternalIndex);
+		ensure(false);
+		return;
+	}
+
+	if (Attachments.HasUnsentAttachments(ENetObjectAttachmentType::Normal, InternalIndex))
 	{
 		UE_LOG(LogIris, Error, TEXT("FReplicationWriter::StartReplication - Expected object %s to not to have any queued up attachments"), *NetRefHandleManager->PrintObjectFromIndex(InternalIndex));
 		ensure(false);
@@ -546,10 +557,17 @@ void FReplicationWriter::StartReplication(uint32 InternalIndex)
 		}
 	}
 
+	const FNetRefHandleManager::FReplicatedObjectData& Data = NetRefHandleManager->GetReplicatedObjectDataNoCheck(InternalIndex);	
+	if (!Data.Protocol)
+	{
+		UE_LOG(LogIris, Error, TEXT("FReplicationWriter::StartReplication - Unexpected call to StartReplication for object ( InternalIndex: %u ) with no ReplicationProtocol"), InternalIndex);
+		ensure(Data.Protocol);
+		return;	
+	}
+
 	// Pending create
 	SetState(InternalIndex, EReplicatedObjectState::PendingCreate);
 
-	const FNetRefHandleManager::FReplicatedObjectData& Data = NetRefHandleManager->GetReplicatedObjectDataNoCheck(InternalIndex);
 	NetRefHandleManager->AddNetObjectRef(InternalIndex);
 
 	Info.ChangeMaskBitCount = Data.Protocol->ChangeMaskBitCount;
