Iris - fix ensure in FSharedStringNetToken.

#Synced-CL Unknown

diff --git a/Engine/Source/Runtime/Net/Iris/Private/Iris/ReplicationSystem/SharedStringNetToken.cpp b/Engine/Source/Runtime/Net/Iris/Private/Iris/ReplicationSystem/SharedStringNetToken.cpp
--- a/Engine/Source/Runtime/Net/Iris/Private/Iris/ReplicationSystem/SharedStringNetToken.cpp
+++ b/Engine/Source/Runtime/Net/Iris/Private/Iris/ReplicationSystem/SharedStringNetToken.cpp
@@ -40,7 +40,7 @@ const TCHAR* FSharedStringNetToken::ResolveToken() const
 
 	const FNetTokenStore* const NetTokenStore = ReplicationSystem->GetNetTokenStore();
 	const FStringTokenStore* const StringTokenStore = NetTokenStore->GetDataStore<FStringTokenStore>();
-	const FNetTokenStoreState* const RemoteNetTokenStoreState = NetTokenStore->GetRemoteNetTokenStoreState(TokenSource.ConnectionId);
+	const FNetTokenStoreState* const RemoteNetTokenStoreState = TokenSource.ConnectionId != InvalidConnectionId ? NetTokenStore->GetRemoteNetTokenStoreState(TokenSource.ConnectionId) : nullptr;
 
 	return StringTokenStore->ResolveToken(Token, RemoteNetTokenStoreState);
 }
@@ -83,7 +83,7 @@ bool FSharedStringNetToken::IsEqualToToken(const FNetToken& OtherToken, const FN
 	else
 	{
 		// Different token stores: resolve the tokens and compare the strings.
-		const FNetTokenStoreState* const ThisRemoteNetTokenStoreState = ThisNetTokenStore->GetRemoteNetTokenStoreState(TokenSource.ConnectionId);
+		const FNetTokenStoreState* const ThisRemoteNetTokenStoreState = TokenSource.ConnectionId != InvalidConnectionId ? ThisNetTokenStore->GetRemoteNetTokenStoreState(TokenSource.ConnectionId) : nullptr;
 		const TCHAR* const ThisResolvedToken = ThisStringTokenStore->ResolveToken(Token, ThisRemoteNetTokenStoreState);
 		const TCHAR* const OtherResolvedToken = OtherStringTokenStore->ResolveToken(OtherToken, OtherResolveContext.RemoteNetTokenStoreState);
 
