Iris - Missing return type is now an error.

#rnx
#Synced-CL 48282875

diff --git a/Engine/Source/Runtime/Net/Iris/IrisCore.Build.cs b/Engine/Source/Runtime/Net/Iris/IrisCore.Build.cs
--- a/Engine/Source/Runtime/Net/Iris/IrisCore.Build.cs
+++ b/Engine/Source/Runtime/Net/Iris/IrisCore.Build.cs
@@ -26,6 +26,7 @@ namespace UnrealBuildTool.Rules
 			}
 
 			CppCompileWarningSettings.UnsafeTypeCastWarningLevel = WarningLevel.Error;
+			CppCompileWarningSettings.ReturnTypeWarningLevel = WarningLevel.Error;
 		}
 	}
 }
