Iris: Shrink string replication, especially ASCII strings

Shrinks length prefix from 16 bits to packed uint16 for all strings.

Shrinks ASCII string content by 12.5% by skipping the high bit of each octet.

[RN] Net/Iris: Optimized string replication.
#Synced-CL 47011050

diff --git a/Engine/Source/Runtime/Net/Iris/Private/Iris/Serialization/NetBitStreamUtil.cpp b/Engine/Source/Runtime/Net/Iris/Private/Iris/Serialization/NetBitStreamUtil.cpp
--- a/Engine/Source/Runtime/Net/Iris/Private/Iris/Serialization/NetBitStreamUtil.cpp
+++ b/Engine/Source/Runtime/Net/Iris/Private/Iris/Serialization/NetBitStreamUtil.cpp
@@ -2,6 +2,8 @@
 
 #include "Iris/Serialization/NetBitStreamUtil.h"
 
+#include "Algo/AllOf.h"
+
 #include "Iris/Core/BitTwiddling.h"
 #include "Iris/Core/IrisLog.h"
 
@@ -15,6 +17,18 @@
 namespace UE::Net
 {
 
+// Whether a string contains only ASCII codepoints.
+static bool ContainsASCIIOnly(FStringView StringView)
+{
+	return Algo::AllOf(
+		StringView,
+		[](FStringView::ElementType Char) -> bool
+		{
+			return (Char & static_cast<FStringView::ElementType>(0x7FU)) == Char;
+		}
+	);
+}
+
 void WritePackedUint64(FNetBitStreamWriter* Writer, uint64 Value)
 {
 	// As we represent the number of bytes to write with three bits we want bits needed to be >= 1 such that the number of bytes ends up in the range [1, 8].
@@ -167,13 +181,13 @@ void WriteString(FNetBitStreamWriter* Writer, FStringView StringView)
 	{
 		UE_LOG(LogSerialization, Error, TEXT("String to write is too long."));
 		Writer->WriteBool(false);
-		Writer->WriteBits(0, 16);
+		WritePackedUint16(Writer, 0);
 		return;
 	}
 
 	constexpr uint32 StackAllocatedBufferSize = 1024;
 	TArray<uint8, TInlineAllocator<StackAllocatedBufferSize>> EncodingBuffer;
-	const bool bIsEncoded = !FCString::IsPureAnsi(StringView.GetData(), StringView.Len());
+	const bool bIsEncoded = !ContainsASCIIOnly(StringView);
 	if (Writer->WriteBool(bIsEncoded))
 	{
 		if (Writer->IsOverflown())
@@ -194,30 +208,20 @@ void WriteString(FNetBitStreamWriter* Writer, FStringView StringView)
 			return;
 		}
 
-		Writer->WriteBits(OutEncodedLength, 16);
+		WritePackedUint16(Writer, static_cast<uint16>(OutEncodedLength));
 		Writer->WriteBitStream(reinterpret_cast<uint32*>(EncodingBuffer.GetData()), 0U, OutEncodedLength*8U);
 	}
 	else
 	{
-		Writer->WriteBits(StringViewLength, 16);
+		WritePackedUint16(Writer, static_cast<uint16>(StringViewLength));
 		if (Writer->IsOverflown())
 		{
 			return;
 		}
 
-		if (StringViewLength > 0)
+		for (TCHAR Char : StringView)
 		{
-			EncodingBuffer.SetNumUninitialized(Align(StringViewLength, SizeAlignment));
-			uint8* EncodingBufferData = EncodingBuffer.GetData();
-			{
-				uint32 CharIndex = 0;
-				for (TCHAR Char : StringView)
-				{
-					EncodingBuffer[CharIndex++] = Char & 0xFF;
-				}
-			}
-
-			Writer->WriteBitStream(reinterpret_cast<uint32*>(EncodingBufferData), 0, StringViewLength*8U);
+			Writer->WriteBits(static_cast<uint32>(Char) & 0x7FU, 7);
 		}
 	}
 }
@@ -234,7 +238,7 @@ void ReadString(FNetBitStreamReader* Reader, FString& OutString)
 	constexpr uint32 StackAllocatedBufferSize = 1024;
 
 	const bool bIsEncoded = Reader->ReadBool();
-	const uint32 Length = Reader->ReadBits(16);
+	const uint32 Length = ReadPackedUint16(Reader);
 
 	if (Reader->IsOverflown())
 	{
@@ -277,21 +281,17 @@ void ReadString(FNetBitStreamReader* Reader, FString& OutString)
 	}
 	else
 	{
-		TArray<char, TInlineAllocator<StackAllocatedBufferSize>> CharBuffer;
-		CharBuffer.SetNumUninitialized(Align(Length + 1, SizeAlignment));
+		OutString.Reset(Length);
+		for (uint32 Index = 0; Index < Length; ++Index)
+		{
+			OutString.AppendChar(static_cast<FString::ElementType>(Reader->ReadBits(7)));
+		}
 
-		// Read ANSI string.
-		Reader->ReadBitStream(reinterpret_cast<uint32*>(CharBuffer.GetData()), Length*8U);
-		// Null-terminate string for easier debugging.
-		CharBuffer[Length] = '\0';
 		if (Reader->IsOverflown())
 		{
-			UE_LOG(LogSerialization, Error, TEXT("Received invalid ANSI string."));
+			UE_LOG(LogSerialization, Error, TEXT("Received truncated ASCII string."));
 			return;
 		}
-
-		// There are two FString constructors with similar signature. The one with known length of the string passes the count first.
-		OutString = FString::ConstructFromPtrSize(CharBuffer.GetData(), static_cast<int32>(Length));
 	}
 }
 
