Make the ThreadBuffer in NetTraceInternal AutoRTFM-safe.

#Synced-CL Unknown

diff --git a/Engine/Source/Runtime/Net/Core/Private/Net/Core/Trace/NetTraceInternal.cpp b/Engine/Source/Runtime/Net/Core/Private/Net/Core/Trace/NetTraceInternal.cpp
--- a/Engine/Source/Runtime/Net/Core/Private/Net/Core/Trace/NetTraceInternal.cpp
+++ b/Engine/Source/Runtime/Net/Core/Private/Net/Core/Trace/NetTraceInternal.cpp
@@ -97,7 +97,7 @@ struct FNetTraceInternal
 
 static std::atomic<uint32> CurrentTraceId(0);
 static std::atomic<UE::Net::FNetDebugNameId> NextNameId(1);
-static thread_local TUniquePtr<FNetTraceInternal::FThreadBuffer> ThreadBuffer;
+static UE_AUTORTFM_DECLARE_THREAD_LOCAL_VAR(TUniquePtr<FNetTraceInternal::FThreadBuffer>, ThreadBuffer);
 static FDelegateHandle OnTraceConnectedHandle;
 static FDelegateHandle OnTraceStoppedHandle;
 static FTransactionallySafeCriticalSection NewObjectsCS;
