Filling Iris object creation headers runs transactionally if needed.

#Synced-CL Unknown

diff --git a/Engine/Source/Runtime/Net/Iris/Private/Iris/ReplicationSystem/NetObjectFactory.cpp b/Engine/Source/Runtime/Net/Iris/Private/Iris/ReplicationSystem/NetObjectFactory.cpp
--- a/Engine/Source/Runtime/Net/Iris/Private/Iris/ReplicationSystem/NetObjectFactory.cpp
+++ b/Engine/Source/Runtime/Net/Iris/Private/Iris/ReplicationSystem/NetObjectFactory.cpp
@@ -11,6 +11,8 @@
 #include "Iris/Serialization/NetBitStreamUtil.h"
 #include "Iris/Serialization/NetSerializationContext.h"
 
+#include "UObject/RemoteExecutor.h"
+
 // Define this on both client and server to add extra header information and quickly detect errors in the header serialization/deserialization code
 
 #include UE_INLINE_GENERATED_CPP_BY_NAME(NetObjectFactory)
@@ -18,6 +20,11 @@
 	#define IRIS_CREATIONHEADER_BITGUARD 0
 #endif
 
+namespace UE::Net::Private
+{
+	static const FName CreateAndFillHeaderTransactionName("UNetObjectFactory::CreateAndFillHeader");
+}
+
 
 //------------------------------------------------------------------------
 // UNetObjectFactory
@@ -55,7 +62,20 @@ TUniquePtr<UE::Net::FNetObjectCreationHeader> UNetObjectFactory::CreateHeader(UE
 	}
 
 	// Ask the derived class to allocate and fill the header
-	TUniquePtr<FNetObjectCreationHeader> Header = CreateAndFillHeader(Handle);
+	TUniquePtr<FNetObjectCreationHeader> Header;
+#if UE_WITH_REMOTE_OBJECT_HANDLE && UE_AUTORTFM
+	if (UE::Net::Private::FBridgeSerialization::IsReplicatingTransactionally(Bridge))
+	{
+		UE::RemoteExecutor::ExecuteTransactional(UE::Net::Private::CreateAndFillHeaderTransactionName, [this, &Header, Handle]()
+			{
+				Header = CreateAndFillHeader(Handle);
+			});
+	}
+	else
+#endif
+	{
+		Header = CreateAndFillHeader(Handle);
+	}
 
 	if (LIKELY(Header.IsValid()))
 	{
diff --git a/Engine/Source/Runtime/Net/Iris/Private/Iris/ReplicationSystem/ObjectReplicationBridge.cpp b/Engine/Source/Runtime/Net/Iris/Private/Iris/ReplicationSystem/ObjectReplicationBridge.cpp
--- a/Engine/Source/Runtime/Net/Iris/Private/Iris/ReplicationSystem/ObjectReplicationBridge.cpp
+++ b/Engine/Source/Runtime/Net/Iris/Private/Iris/ReplicationSystem/ObjectReplicationBridge.cpp
@@ -567,6 +567,11 @@ void UObjectReplicationBridge::PreRegisterObjectWithReferenceHandle(const UObjec
 	ObjectReferenceCache->AddPreRegisteredReference(Handle, Instance, ReplicationSystem);
 }
 
+bool UObjectReplicationBridge::IsReplicatingTransactionally() const
+{
+	return GetReplicationSystem()->IsReplicatingTransactionally();
+}
+
 UE::Net::FNetRefHandle UObjectReplicationBridge::StartReplicatingNetObject(UObject* Instance, UE::Net::EReplicationFragmentTraits Traits, UE::Net::FNetObjectFactoryId NetFactoryId, bool bForceNewHandle)
 {
 	using namespace UE::Net;
diff --git a/Engine/Source/Runtime/Net/Iris/Public/Iris/ReplicationSystem/ObjectReplicationBridge.h b/Engine/Source/Runtime/Net/Iris/Public/Iris/ReplicationSystem/ObjectReplicationBridge.h
--- a/Engine/Source/Runtime/Net/Iris/Public/Iris/ReplicationSystem/ObjectReplicationBridge.h
+++ b/Engine/Source/Runtime/Net/Iris/Public/Iris/ReplicationSystem/ObjectReplicationBridge.h
@@ -566,6 +566,9 @@ private:
 	/** Invoked right before we apply the state for a new received subobject but after we have applied state on the root object. */
 	void SubObjectCreatedFromReplication(UE::Net::Private::FInternalNetRefIndex RootObjectIndex, FNetRefHandle SubObjectCreated);
 
+	/** Gets the bReplicateTransactionally flag from the ReplicationSystem */
+	bool IsReplicatingTransactionally() const;
+
 private:
 
 	friend UE::Net::Private::FObjectPoller;
@@ -776,6 +779,11 @@ public:
 	{
 		return Bridge->SubObjectCreatedFromReplication(RootObjectIndex, SubObjectCreated);
 	}
+
+	static bool IsReplicatingTransactionally(const UObjectReplicationBridge* Bridge)
+	{
+		return Bridge->IsReplicatingTransactionally();
+	}
 };
 
 } // end namespace UE::Net::Private
\ No newline at end of file
