FNetBitArray's ForAllSetBitsInRange is now handling invalid StartBit and EndBit parameters more gracefully.

#Synced-CL 48452653

diff --git a/Engine/Source/Runtime/Net/Core/Public/Net/Core/NetBitArray.h b/Engine/Source/Runtime/Net/Core/Public/Net/Core/NetBitArray.h
--- a/Engine/Source/Runtime/Net/Core/Public/Net/Core/NetBitArray.h
+++ b/Engine/Source/Runtime/Net/Core/Public/Net/Core/NetBitArray.h
@@ -285,7 +285,7 @@ public:
 	template<typename T>
 	void ForAllSetBits(T&& Functor) const;
 
-	/** Iterate over all set bits and invoke functor with signature void(uint32 BitIndex). */
+	/** Iterate over all set bits and invoke functor with signature void(uint32 BitIndex). EndBit is clamped to the the number of bits in the array minus one. */
 	template<typename T>
 	void ForAllSetBitsInRange(const uint32 StartBit, const uint32 EndBit, T&& Functor) const;
 
@@ -521,7 +521,7 @@ public:
 	template <typename T>
 	void ForAllSetBits(T&& Functor) const;
 
-	/** Iterate over all set bits from StartBit to EndBit (inclusive) and invoke functor with signature void(uint32 BitIndex). */
+	/** Iterate over all set bits from StartBit to EndBit (inclusive) and invoke functor with signature void(uint32 BitIndex). EndBit is clamped to the the number of bits in the array minus one. */
 	template <typename T>
 	void ForAllSetBitsInRange(const uint32 StartBit, const uint32 EndBit, T&& Functor) const;
 
@@ -906,10 +906,14 @@ private:
 	}
 
 	//Calls Functor for every set bit from StartBit to EndBit (inclusive)
-	template<typename T> static void ForAllSetBitsInRange(const StorageWordType* Storage, const uint32 TotalWordCount, const uint32 TotalBitCount, const uint32 StartBit, const uint32 EndBit, T&& Functor)
+	template<typename T> static void ForAllSetBitsInRange(const StorageWordType* Storage, const uint32 TotalWordCount, const uint32 TotalBitCount, const uint32 StartBit, const uint32 InEndBit, T&& Functor)
 	{
-		check(StartBit < TotalBitCount);
-		check(EndBit < TotalBitCount);
+		const uint32 EndBit = FPlatformMath::Min(InEndBit, TotalBitCount - 1U);
+		if ((TotalBitCount == 0U) || (StartBit >= TotalBitCount) || (EndBit < StartBit))
+		{
+			return;
+		}
+
 		const uint32 StartWord = StartBit / WordBitCount;
 		const uint32 EndWord = EndBit / WordBitCount;
 		const uint32 StartWordBitIndex = StartWord * WordBitCount;//Beginning bit of the StartWord (may be lower than StartBit)
