Iris
* Added compil-time function to calculate bits needed to serialize a number

#rnx
#Synced-CL 48500075

diff --git a/Engine/Source/Runtime/Net/Iris/Public/Iris/Core/BitTwiddling.h b/Engine/Source/Runtime/Net/Iris/Public/Iris/Core/BitTwiddling.h
--- a/Engine/Source/Runtime/Net/Iris/Public/Iris/Core/BitTwiddling.h
+++ b/Engine/Source/Runtime/Net/Iris/Public/Iris/Core/BitTwiddling.h
@@ -28,7 +28,7 @@ namespace UE::Net
  *         For signed integers a sign bit is assumed to be propagated to all higher bits.
  */
 template<typename T, typename TEnableIf<TIsSigned<T>::Value && sizeof(T) <= 4U, int32>::Type X = -1>
-uint32 GetBitsNeeded(const T Value)
+constexpr uint32 GetBitsNeeded(const T Value)
 {
 	typedef typename TUnsignedIntType<sizeof(T)>::Type SmallUnsignedType;
 
@@ -47,20 +47,20 @@ uint32 GetBitsNeeded(const T Value)
 }
 
 template<typename T, typename TEnableIf<!TIsSigned<T>::Value && sizeof(T) <= 4U, uint32>::Type X = 1U>
-uint32 GetBitsNeeded(const T Value)
+constexpr uint32 GetBitsNeeded(const T Value)
 {
 	return 32U - static_cast<uint32>(FPlatformMath::CountLeadingZeros(uint32(Value)));
 }
 
 template<typename T, typename TEnableIf<TIsSigned<T>::Value && sizeof(T) == 8U, int64>::Type X = -1LL>
-uint32 GetBitsNeeded(const T Value)
+constexpr uint32 GetBitsNeeded(const T Value)
 {
 	const uint64 MassagedValue = uint64(Value ^ (Value >> 63U));
 	return 65U - static_cast<uint32>(FPlatformMath::CountLeadingZeros64(MassagedValue));
 }
 
 template<typename T, typename TEnableIf<!TIsSigned<T>::Value && sizeof(T) == 8U, uint64>::Type X = 1ULL>
-uint32 GetBitsNeeded(const T Value)
+constexpr uint32 GetBitsNeeded(const T Value)
 {
 	return 64U - static_cast<uint32>(FPlatformMath::CountLeadingZeros64(Value));
 }
@@ -75,7 +75,7 @@ uint32 GetBitsNeeded(const T Value)
  * @note It is assumed UpperBound >= LowerBound.
  */
 template<typename T, typename TEnableIf<TIsSigned<T>::Value && sizeof(T) <= 4U, int32>::Type X = -1>
-uint32 GetBitsNeededForRange(const T LowerBound, const T UpperBound)
+constexpr uint32 GetBitsNeededForRange(const T LowerBound, const T UpperBound)
 {
 	using UnsignedType = typename TUnsignedIntType<sizeof(T)>::Type;
 	const uint32 Range = uint32(UnsignedType(UnsignedType(UpperBound) - UnsignedType(LowerBound)));
@@ -83,14 +83,14 @@ uint32 GetBitsNeededForRange(const T LowerBound, const T UpperBound)
 }
 
 template<typename T, typename TEnableIf<!TIsSigned<T>::Value && sizeof(T) <= 4U, uint32>::Type X = 1U>
-uint32 GetBitsNeededForRange(const T LowerBound, const T UpperBound)
+constexpr uint32 GetBitsNeededForRange(const T LowerBound, const T UpperBound)
 {
 	const uint32 Range = uint32(UpperBound) - uint32(LowerBound);
 	return 32U - static_cast<uint32>(FPlatformMath::CountLeadingZeros(Range));
 }
 
 template<typename T, typename TEnableIf<sizeof(T) == 8U, uint64>::Type X = 1ULL>
-uint32 GetBitsNeededForRange(const T LowerBound, const T UpperBound)
+constexpr uint32 GetBitsNeededForRange(const T LowerBound, const T UpperBound)
 {
 	const uint64 Range = uint64(UpperBound) - uint64(LowerBound);
 	return 64U - static_cast<uint32>(FPlatformMath::CountLeadingZeros64(Range));
@@ -98,7 +98,7 @@ uint32 GetBitsNeededForRange(const T LowerBound, const T UpperBound)
 
 /** GetLeastSignificantBit returns the least significant bit set in Value, or 0 if none is set. */
 template<typename T, typename TEnableIf<!TIsSigned<T>::Value, int>::Type X = 1>
-T GetLeastSignificantBit(const T Value)
+constexpr T GetLeastSignificantBit(const T Value)
 {
 	using SignedT = typename TSignedIntType<sizeof(T)>::Type;
 	const T LeastSignificantBit = Value & T(-SignedT(Value));
diff --git a/Engine/Source/Runtime/Net/Iris/Public/Iris/ReplicationSystem/NetObjectFactoryRegistry.h b/Engine/Source/Runtime/Net/Iris/Public/Iris/ReplicationSystem/NetObjectFactoryRegistry.h
--- a/Engine/Source/Runtime/Net/Iris/Public/Iris/ReplicationSystem/NetObjectFactoryRegistry.h
+++ b/Engine/Source/Runtime/Net/Iris/Public/Iris/ReplicationSystem/NetObjectFactoryRegistry.h
@@ -6,6 +6,7 @@
 #include "UObject/Class.h"
 #include "UObject/NameTypes.h"
 #include "UObject/WeakObjectPtrTemplates.h"
+#include "Iris/Core/BitTwiddling.h"
 
 // Define this to increase the maximum amount of NetObjectFactories that can exist
 #ifndef UE_IRIS_MAX_NETOBJECT_FACTORIES
@@ -18,6 +19,8 @@ namespace UE::Net
 typedef uint8 FNetObjectFactoryId;
 enum { InvalidNetObjectFactoryId = MAX_uint8 };
 
+static_assert((uint32)UE_IRIS_MAX_NETOBJECT_FACTORIES < (uint32)InvalidNetObjectFactoryId, "We currently don't support more then 255 unique factories");
+
 
 /**
  * Keeps track of Iris NetObjectFactory templates
@@ -45,7 +48,7 @@ public:
 	static constexpr uint32 MaxFactories = UE_IRIS_MAX_NETOBJECT_FACTORIES;
 
 	/** The amount of bits to serialize FNetFactoryIds with  */
-	static constexpr uint32 GetMaxBits() { return GetNumBits(MaxFactories-1); }
+	static consteval uint32 GetMaxBits() { return UE::Net::GetBitsNeeded(MaxFactories-1); }
 
 	struct FFactoryData
 	{
@@ -65,11 +68,6 @@ private:
 	FNetObjectFactoryRegistry() = delete;
 	~FNetObjectFactoryRegistry() = delete;
 
-	static constexpr uint32 GetNumBits(uint32 Number)
-	{
-		return Number==0 ? 0 : 1 + GetNumBits(Number >> 1);
-	}
-
 private:
 
 	static TArray<FFactoryData, TFixedAllocator<MaxFactories>> NetFactories;
