Iris: converted IRIS_PROFILER_PROTOCOL_NAME profiling macros to avoid doing string copies, since FNetDebugName store their debug names internally.

#Synced-CL 48523383

diff --git a/Engine/Source/Runtime/Net/Iris/Public/Iris/Core/IrisProfiler.h b/Engine/Source/Runtime/Net/Iris/Public/Iris/Core/IrisProfiler.h
--- a/Engine/Source/Runtime/Net/Iris/Public/Iris/Core/IrisProfiler.h
+++ b/Engine/Source/Runtime/Net/Iris/Public/Iris/Core/IrisProfiler.h
@@ -37,12 +37,16 @@
 #		pragma comment (lib, "c:/Program Files/Superluminal/Performance/API/lib/x64/PerformanceAPI_MD.lib")
 #		define IRIS_PROFILER_SCOPE(X) PERFORMANCEAPI_INSTRUMENT(UE_STRINGIZE(X))
 #		define IRIS_PROFILER_SCOPE_CONDITIONAL(X,Cond) PERFORMANCEAPI_INSTRUMENT(UE_STRINGIZE(X))
+#		define IRIS_PROFILER_SCOPE_STR(X) PERFORMANCEAPI_INSTRUMENT_DATA(UE_JOIN(IrisProfilerScope, __LINE__), X)
+#		define IRIS_PROFILER_SCOPE_STR_CONDITIONAL(X, Cond) PERFORMANCEAPI_INSTRUMENT_DATA(UE_JOIN(IrisProfilerScope, __LINE__), X)
 #		define IRIS_PROFILER_SCOPE_TEXT(X) PERFORMANCEAPI_INSTRUMENT_DATA(UE_JOIN(IrisProfilerScope, __LINE__), X)
 #		define IRIS_PROFILER_SCOPE_TEXT_CONDITIONAL(X, Cond) PERFORMANCEAPI_INSTRUMENT_DATA(UE_JOIN(IrisProfilerScope, __LINE__), X)
 #	else
 #		include "ProfilingDebugging/CpuProfilerTrace.h"
 #		define IRIS_PROFILER_SCOPE(X) TRACE_CPUPROFILER_EVENT_SCOPE(X)
 #		define IRIS_PROFILER_SCOPE_CONDITIONAL(X,Cond) TRACE_CPUPROFILER_EVENT_SCOPE_CONDITIONAL(X, Cond)
+#		define IRIS_PROFILER_SCOPE_STR(X) TRACE_CPUPROFILER_EVENT_SCOPE_STR(X)
+#		define IRIS_PROFILER_SCOPE_STR_CONDITIONAL(X, Cond) TRACE_CPUPROFILER_EVENT_SCOPE_STR_CONDITIONAL(X, Cond)
 #		define IRIS_PROFILER_SCOPE_TEXT(X) TRACE_CPUPROFILER_EVENT_SCOPE_TEXT(X)
 #		define IRIS_PROFILER_SCOPE_TEXT_CONDITIONAL(X, Cond) TRACE_CPUPROFILER_EVENT_SCOPE_TEXT_CONDITIONAL(X, Cond)
 #	endif
@@ -50,13 +54,15 @@
 #	define PERFORMANCEAPI_ENABLED 0
 #	define IRIS_PROFILER_SCOPE(X)
 #	define IRIS_PROFILER_SCOPE_CONDITIONAL(X, Cond)
+#	define IRIS_PROFILER_SCOPE_STR(X)
+#	define IRIS_PROFILER_SCOPE_STR_CONDITIONAL(X, Cond)
 #	define IRIS_PROFILER_SCOPE_TEXT(X)
 #	define IRIS_PROFILER_SCOPE_TEXT_CONDITIONAL(X, Cond)
 #endif
 
 #if UE_IRIS_PROFILER_ENABLE_PROTOCOL_NAMES
-#	define IRIS_PROFILER_PROTOCOL_NAME(X) IRIS_PROFILER_SCOPE_TEXT(X)
-#	define IRIS_PROFILER_PROTOCOL_NAME_CONDITIONAL(X, Cond) IRIS_PROFILER_SCOPE_TEXT_CONDITIONAL(X, Cond)
+#	define IRIS_PROFILER_PROTOCOL_NAME(X) IRIS_PROFILER_SCOPE_STR(X)
+#	define IRIS_PROFILER_PROTOCOL_NAME_CONDITIONAL(X, Cond) IRIS_PROFILER_SCOPE_STR_CONDITIONAL(X, Cond)
 #else
 #	define IRIS_PROFILER_PROTOCOL_NAME(X)
 #	define IRIS_PROFILER_PROTOCOL_NAME_CONDITIONAL(X, Cond)
