Iris
* Cleanup ReplicationReader by removing two deprecated codepaths
* "net.Iris.UseOptObjectRefTracking true" is now the the only way to resolve references
* "net.Iris.DeferEndReplication false" is no longer supported

#rnx
#Synced-CL 48467961

#Synced-CL 48149700

diff --git a/Engine/Source/Runtime/Net/Iris/Private/Iris/ReplicationSystem/ReplicationReader.cpp b/Engine/Source/Runtime/Net/Iris/Private/Iris/ReplicationSystem/ReplicationReader.cpp
--- a/Engine/Source/Runtime/Net/Iris/Private/Iris/ReplicationSystem/ReplicationReader.cpp
+++ b/Engine/Source/Runtime/Net/Iris/Private/Iris/ReplicationSystem/ReplicationReader.cpp
@@ -75,27 +75,12 @@ static FAutoConsoleVariableRef CVarColdResolvingRetryTimeMS(
 	ColdResolvingRetryTimeMS,
 	TEXT("Resolve unresolved cold references after this many milliseconds."));
 
-// When true, the function UpdateObjectReferenceTracking_Fast() is used and UpdateObjectReferenceTracking()
-// otherwise. Once UpdateObjectReferenceTracking_Fast() has been tested sufficiently it will become the only
-// version of this function and the cvar removed.
-static bool bUseOptObjectRefTracking = false;
-static FAutoConsoleVariableRef CVarUseOptObjectRefTracking(
-	TEXT("net.Iris.UseOptObjectRefTracking"),
-	bUseOptObjectRefTracking,
-	TEXT("Use a more optimized version of FReplicationReader::UpdateObjectReferenceTracking()."));
-
 static bool bExecuteReliableRPCsBeforeApplyState = true;
 static FAutoConsoleVariableRef CVarExecuteReliableRPCsBeforeApplyState(
 		TEXT("net.Iris.ExecuteReliableRPCsBeforeApplyState"),
 		bExecuteReliableRPCsBeforeApplyState,
 		TEXT("If true and Iris runs in backwards compatibility mode then reliable RPCs will be executed before we apply state data on the target object unless we first need to spawn the object."));
 
-static bool bDeferEndReplication = true;
-static FAutoConsoleVariableRef CVarDeferEndReplication(
-	TEXT("net.Iris.DeferEndReplication"),
-	bDeferEndReplication,
-	TEXT("if true calls to EndReplication will be defered until after we have applied statedata. Default is true."));
-
 static bool bImmediateDispatchEndReplicationForSubObjects = true;
 static FAutoConsoleVariableRef CVarImmediateDispatchEndReplicationForSubObjects(
 	TEXT("net.Iris.ImmediateDispatchEndReplicationForSubObjects"),
@@ -531,27 +516,20 @@ uint32 FReplicationReader::ReadObjectsAndSubObjectsPendingDestroy(FNetSerializat
 				UE_NET_TRACE_SET_SCOPE_OBJECTID(DestroyedObjectScope, IncompleteHandle);
 
 				// Defer EndReplication until after applying state data
-				if (bDeferEndReplication)
-				{
-					FDispatchObjectInfo& Info = ObjectsToDispatchArray->AddPendingDispatchObjectInfo(TempLinearAllocator);
-
-					Info.bDestroy = bShouldDestroyInstance;
-					Info.bTearOff = false;
-					Info.bDeferredEndReplication = true;
-					Info.InternalIndex = InternalIndex;
-					Info.bIsInitialState = false;
-					Info.bHasState = false;
-					Info.bHasAttachments = false;
-					Info.bShouldCallSubObjectCreatedFromReplication = false;
-					Info.CloseReason = CloseReason;
-
-					// Mark for dispatch
-					ObjectsToDispatchArray->CommitPendingDispatchObjectInfo();
-				}
-				else
-				{
-					EndReplication(InternalIndex, false, bShouldDestroyInstance, CloseReason);
-				}
+				FDispatchObjectInfo& Info = ObjectsToDispatchArray->AddPendingDispatchObjectInfo(TempLinearAllocator);
+
+				Info.bDestroy = bShouldDestroyInstance;
+				Info.bTearOff = false;
+				Info.bDeferredEndReplication = true;
+				Info.InternalIndex = InternalIndex;
+				Info.bIsInitialState = false;
+				Info.bHasState = false;
+				Info.bHasAttachments = false;
+				Info.bShouldCallSubObjectCreatedFromReplication = false;
+				Info.CloseReason = CloseReason;
+
+				// Mark for dispatch
+				ObjectsToDispatchArray->CommitPendingDispatchObjectInfo();
 			}
 			else
 			{
@@ -612,36 +590,21 @@ uint32 FReplicationReader::ReadRootObjectsPendingDestroy(FNetSerializationContex
 				UE_NET_TRACE_SET_SCOPE_OBJECTID(DestroyedObjectScope, IncompleteHandle);
 
 				// Defer EndReplication until after applying state data
-				if (bDeferEndReplication)
-				{
-					FDispatchObjectInfo& Info = ObjectsToDispatchArray->AddPendingDispatchObjectInfo(TempLinearAllocator);
-
-					Info.bDestroy = bShouldDestroyInstance;
-					Info.bTearOff = false;
-					Info.bDeferredEndReplication = true;
-					Info.bProcessEndReplicationForHierarchy = true;
-					Info.InternalIndex = InternalIndex;
-					Info.bIsInitialState = false;
-					Info.bHasState = false;
-					Info.bHasAttachments = false;
-					Info.bShouldCallSubObjectCreatedFromReplication = false;
-					Info.CloseReason = CloseReason;
-
-					// Mark for dispatch
-					ObjectsToDispatchArray->CommitPendingDispatchObjectInfo();
-				}
-				else
-				{
-					constexpr bool bTearOff = false;
-					EndReplication(InternalIndex, bTearOff, bShouldDestroyInstance, CloseReason);
-					for (FInternalNetRefIndex SubObjectIndex : NetRefHandleManager->GetSubObjects(InternalIndex))
-					{
-						// End replication for all subobjects and destroy dynamic ones
-						const FNetRefHandle SubObjectRefHandle = NetRefHandleManager->GetNetRefHandleFromInternalIndex(SubObjectIndex);
-						const bool bShouldDestroySubObject = SubObjectRefHandle.IsDynamic();
-						EndReplication(SubObjectIndex, bTearOff, bShouldDestroySubObject, CloseReason);
-					}
-				}
+				FDispatchObjectInfo& Info = ObjectsToDispatchArray->AddPendingDispatchObjectInfo(TempLinearAllocator);
+
+				Info.bDestroy = bShouldDestroyInstance;
+				Info.bTearOff = false;
+				Info.bDeferredEndReplication = true;
+				Info.bProcessEndReplicationForHierarchy = true;
+				Info.InternalIndex = InternalIndex;
+				Info.bIsInitialState = false;
+				Info.bHasState = false;
+				Info.bHasAttachments = false;
+				Info.bShouldCallSubObjectCreatedFromReplication = false;
+				Info.CloseReason = CloseReason;
+
+				// Mark for dispatch
+				ObjectsToDispatchArray->CommitPendingDispatchObjectInfo();
 			}
 			else
 			{
@@ -1608,153 +1571,6 @@ ErrorHandling:
 }
 
 // Update reference tracking maps for the current object. It is assumed the ObjectReferenceTracker do no include duplicates for a given key.
-void FReplicationReader::UpdateObjectReferenceTracking(FReplicatedObjectInfo* ReplicationInfo, FNetBitArrayView ChangeMask, bool bIncludeInitState, FResolvedNetRefHandlesArray& OutNewResolvedRefHandles, const FObjectReferenceTracker& NewUnresolvedReferences, const FObjectReferenceTracker& NewMappedDynamicReferences)
-{
-	IRIS_PROFILER_SCOPE(FReplicationReader_UpdateObjectReferenceTracking)
-
-	/*
-	 * As we store references per changemask we need to construct a set of all unresolved references and
-	 * compare with the new set of unresolved references. The new set is found by first updating the 
-	 * references that were found in the changemask.
-	 */
-	{
-		// Try to avoid dynamic allocations during the update of the UnresolvedObjectReferences.
-		ReplicationInfo->UnresolvedObjectReferences.Reserve(ReplicationInfo->UnresolvedObjectReferences.Num() + NewUnresolvedReferences.Num());
-
-		TSet<FNetRefHandle> OldUnresolvedSet;
-		OldUnresolvedSet.Reserve(ReplicationInfo->UnresolvedObjectReferences.Num());
-		for (const FObjectReferenceTracker::ElementType& Element : ReplicationInfo->UnresolvedObjectReferences)
-		{
-			OldUnresolvedSet.Add(Element.Value);
-		}
-
-		// Replace each entry in UnresolvedObjectReferences for the given changemask
-		auto UpdateUnresolvedReferencesForChange = [ReplicationInfo, &NewUnresolvedReferences](uint32 ChangeBit)
-		{
-			FObjectReferenceTracker& UnresolvedObjectReferences = ReplicationInfo->UnresolvedObjectReferences;
-
-			// Use the Hash variants of the map methods to provide a tiny speedup.
-			const uint32 KeyHash = ::GetTypeHash(ChangeBit);
-			UnresolvedObjectReferences.RemoveByHash(KeyHash, ChangeBit);
-
-			for (auto It = NewUnresolvedReferences.CreateConstKeyIterator(ChangeBit); It; ++It)
-			{
-				const FNetRefHandle RefHandle = It.Value();
-				UnresolvedObjectReferences.AddByHash(KeyHash, ChangeBit, RefHandle);
-			}
-		};
-
-		ChangeMask.ForAllSetBits(UpdateUnresolvedReferencesForChange);
-		if (bIncludeInitState)
-		{
-			UpdateUnresolvedReferencesForChange(FakeInitChangeMaskOffset);
-		}
-
-		// The unresolved set is now updated with the current status of unresolved references.
-		TSet<FNetRefHandle> NewUnresolvedSet;
-		NewUnresolvedSet.Reserve(NewUnresolvedReferences.Num());
-		for (const FObjectReferenceTracker::ElementType& Element : ReplicationInfo->UnresolvedObjectReferences)
-		{
-			NewUnresolvedSet.Add(Element.Value);
-		}
-
-		// Update ReplicationInfo with the status of unresolved references
-		ReplicationInfo->bHasUnresolvedReferences = NewUnresolvedSet.Num() > 0;
-		ReplicationInfo->bHasUnresolvedInitialReferences = ReplicationInfo->UnresolvedObjectReferences.Find(FakeInitChangeMaskOffset) != nullptr;
-
-		// Remove resolved or no longer existing references
-		const uint32 OwnerInternalIndex = ReplicationInfo->InternalIndex;
-		for (FNetRefHandle Handle : OldUnresolvedSet)
-		{
-			if (!NewUnresolvedSet.Contains(Handle))
-			{
-				// Store new resolved handles so we can update partially resolved references properly
-				OutNewResolvedRefHandles.Add(Handle);
-
-				// Remove from tracking
-				UnresolvedHandleToDependents.RemoveSingle(Handle, OwnerInternalIndex);
-				RemoveFromUnresolvedCache(Handle);
-				UE_LOG(LogIris, Verbose, TEXT("FReplicationReader::UpdateObjectReferenceTracking Removing unresolved reference %s for %s"), ToCStr(Handle.ToString()), ToCStr(NetRefHandleManager->GetNetRefHandleFromInternalIndex(OwnerInternalIndex).ToString()));
-			}
-		}
-
-		// Add new unresolved references
-		for (FNetRefHandle Handle : NewUnresolvedSet)
-		{
-			if (!OldUnresolvedSet.Contains(Handle))
-			{
-				// Add to tracking
-				UnresolvedHandleToDependents.Add(Handle, OwnerInternalIndex);
-				UE_LOG(LogIris, Verbose, TEXT("FReplicationReader::UpdateObjectReferenceTracking Adding unresolved reference %s for %s"), ToCStr(Handle.ToString()), ToCStr(NetRefHandleManager->GetNetRefHandleFromInternalIndex(OwnerInternalIndex).ToString()));
-			}
-		}
-	}
-
-	// Update tracking for resolved dynamic references
-	if (bRemapDynamicObjects)
-	{
-		// Try to avoid dynamic allocations during the update of the ResolvedDynamicObjectReferences.
-		ReplicationInfo->ResolvedDynamicObjectReferences.Reserve(ReplicationInfo->ResolvedDynamicObjectReferences.Num() + NewMappedDynamicReferences.Num());
-
-		TSet<FNetRefHandle> OldResolvedSet;
-		OldResolvedSet.Reserve(ReplicationInfo->ResolvedDynamicObjectReferences.Num());
-		for (const FObjectReferenceTracker::ElementType& Element : ReplicationInfo->ResolvedDynamicObjectReferences)
-		{
-			OldResolvedSet.Add(Element.Value);
-		}
-
-		// Replace each entry in ResolvedDynamicObjectReferences for the given changemask
-		auto UpdateResolvedReferencesForChange = [ReplicationInfo, &NewMappedDynamicReferences](uint32 ChangeBit)
-		{
-			FObjectReferenceTracker& ResolvedDynamicObjectReferences = ReplicationInfo->ResolvedDynamicObjectReferences;
-
-			const uint32 KeyHash = ::GetTypeHash(ChangeBit);
-			ResolvedDynamicObjectReferences.RemoveByHash(KeyHash, ChangeBit);
-
-			for (auto It = NewMappedDynamicReferences.CreateConstKeyIterator(ChangeBit); It; ++It)
-			{
-				const FNetRefHandle RefHandle = It.Value();
-				ResolvedDynamicObjectReferences.AddByHash(KeyHash, ChangeBit, RefHandle);
-			}
-		};
-
-		ChangeMask.ForAllSetBits(UpdateResolvedReferencesForChange);
-		// Intentionally leaving out init state. It seems weird to call rep notifies and update init only properties after
-		// the initial state has already been applied.
-
-		// The resolved set is now updated with the current status of resolved references.
-		TSet<FNetRefHandle> NewResolvedSet;
-		NewResolvedSet.Reserve(ReplicationInfo->ResolvedDynamicObjectReferences.Num());
-		for (const FObjectReferenceTracker::ElementType& Element : ReplicationInfo->ResolvedDynamicObjectReferences)
-		{
-			NewResolvedSet.Add(Element.Value);
-		}
-
-		// Remove now unresolved or no longer existing references
-		const uint32 OwnerInternalIndex = ReplicationInfo->InternalIndex;
-		for (FNetRefHandle Handle : OldResolvedSet)
-		{
-			if (!NewResolvedSet.Contains(Handle))
-			{
-				// Remove from tracking
-				ResolvedDynamicHandleToDependents.RemoveSingle(Handle, OwnerInternalIndex);
-				//UE_LOG(LogIris, Verbose, TEXT("FReplicationReader::UpdateObjectReferenceTracking Removing resolved dynamic reference %s for %s"), ToCStr(Handle.ToString()), ToCStr(NetRefHandleManager->GetNetRefHandleFromInternalIndex(OwnerInternalIndex).ToString()));
-			}
-		}
-
-		// Add new resolved dynamic references
-		for (FNetRefHandle Handle : NewResolvedSet)
-		{
-			if (Handle.IsDynamic() && !OldResolvedSet.Contains(Handle))
-			{
-				// Add to tracking
-				ResolvedDynamicHandleToDependents.Add(Handle, OwnerInternalIndex);
-				//UE_LOG(LogIris, Verbose, TEXT("FReplicationReader::UpdateObjectReferenceTracking Adding resolved dynamic reference %s for %s"), ToCStr(Handle.ToString()), ToCStr(NetRefHandleManager->GetNetRefHandleFromInternalIndex(OwnerInternalIndex).ToString()));
-			}
-		}
-	}
-}
-
 void FReplicationReader::UpdateObjectReferenceTracking_Fast(FReplicatedObjectInfo* ReplicationInfo, FNetBitArrayView ChangeMask, bool bIncludeInitState, FResolvedNetRefHandlesArray& OutNewResolvedRefHandles, const FObjectReferenceTracker& NewUnresolvedReferences, const FObjectReferenceTracker& NewMappedDynamicReferences)
 {
 	IRIS_PROFILER_SCOPE(FReplicationReader_UpdateObjectReferenceTracking);
@@ -1924,10 +1740,7 @@ void FReplicationReader::RemoveUnresolvedObjectReferenceInReplicationInfo(FRepli
 		{
 			It.RemoveCurrent();
 
-			if (bUseOptObjectRefTracking)
-			{
-				ReplicationInfo->RemoveUnresolvedHandleCount(RefHandle);
-			}
+			ReplicationInfo->RemoveUnresolvedHandleCount(RefHandle);
 		}
 	}
 }
@@ -1941,10 +1754,7 @@ void FReplicationReader::RemoveResolvedObjectReferenceInReplicationInfo(FReplica
 		{
 			It.RemoveCurrent();
 
-			if (bUseOptObjectRefTracking)
-			{
-				ReplicationInfo->RemoveResolvedDynamicHandleCount(RefHandle);
-			}
+			ReplicationInfo->RemoveResolvedDynamicHandleCount(RefHandle);
 		}
 	}
 }
@@ -1979,20 +1789,14 @@ bool FReplicationReader::MoveResolvedObjectReferenceToUnresolvedInReplicationInf
 			// At this point we'd like to skip iteration to the next key as a handle can only be found once per changemask.
 			It.RemoveCurrent();
 
-			if (bUseOptObjectRefTracking)
-			{
-				ReplicationInfo->RemoveResolvedDynamicHandleCount(UnresolvableHandle);
-			}
+			ReplicationInfo->RemoveResolvedDynamicHandleCount(UnresolvableHandle);
 
 			// This handle should only have existed once in the ResolvedDynamicObjectReferences map and should not be able to
 			// already exist in the UnresolvedObjectReferences map, so no need to call AddUnique.
 			UnresolvedObjectReferences.Add(ChangemaskOffset, UnresolvableHandle);
 			
-			if (bUseOptObjectRefTracking)
-			{
-				int16& UnresolvableHandleCount = UnresolvedHandleCount.FindOrAdd(UnresolvableHandle, 0);
-				UnresolvableHandleCount++;
-			}
+			int16& UnresolvableHandleCount = UnresolvedHandleCount.FindOrAdd(UnresolvableHandle, 0);
+			UnresolvableHandleCount++;
 
 			UE_LOG(LogIris, Verbose, TEXT("FReplicationReader::MoveResolvedObjectReferenceToUnresolvedInReplicationInfo Moving from resolved to unresolved reference %s for %s"), ToCStr(UnresolvableHandle.ToString()), ToCStr(NetRefHandleManager->GetNetRefHandleFromInternalIndex(ReplicationInfo->InternalIndex).ToString()));
 		}
@@ -2070,14 +1874,7 @@ void FReplicationReader::BuildUnresolvedChangeMaskAndUpdateObjectReferenceTracki
 	}
 
 	// Update object specific
-	if (bUseOptObjectRefTracking)
-	{
-		UpdateObjectReferenceTracking_Fast(ReplicationInfo, CollectorChangeMask, Collector.IsInitStateIncluded(), OutNewResolvedRefHandles, UnresolvedReferencesCache, MappedDynamicReferencesCache);
-	}
-	else
-	{
-		UpdateObjectReferenceTracking(ReplicationInfo, CollectorChangeMask, Collector.IsInitStateIncluded(), OutNewResolvedRefHandles, UnresolvedReferencesCache, MappedDynamicReferencesCache);
-	}
+	UpdateObjectReferenceTracking_Fast(ReplicationInfo, CollectorChangeMask, Collector.IsInitStateIncluded(), OutNewResolvedRefHandles, UnresolvedReferencesCache, MappedDynamicReferencesCache);
 }
 
 void FReplicationReader::ResolveAndDispatchUnresolvedReferencesForObject(FNetSerializationContext& Context, uint32 InternalIndex)
diff --git a/Engine/Source/Runtime/Net/Iris/Private/Iris/ReplicationSystem/ReplicationReader.h b/Engine/Source/Runtime/Net/Iris/Private/Iris/ReplicationSystem/ReplicationReader.h
--- a/Engine/Source/Runtime/Net/Iris/Private/Iris/ReplicationSystem/ReplicationReader.h
+++ b/Engine/Source/Runtime/Net/Iris/Private/Iris/ReplicationSystem/ReplicationReader.h
@@ -198,9 +198,6 @@ private:
 	const FReplicatedObjectInfo* GetReplicatedObjectInfo(uint32 InternalIndex) const;
 
 	// Update reference tracking maps for the current object
-	void UpdateObjectReferenceTracking(FReplicatedObjectInfo* ReplicationInfo, FNetBitArrayView ChangeMask, bool bIncludeInitState, FResolvedNetRefHandlesArray& OutNewResolvedRefHandles, const FObjectReferenceTracker& NewUnresolvedReferences, const FObjectReferenceTracker& NewMappedDynamicReferences);
-	
-	// An optimized version of UpdateObjectReferenceTracking().
 	void UpdateObjectReferenceTracking_Fast(FReplicatedObjectInfo* ReplicationInfo, FNetBitArrayView ChangeMask, bool bIncludeInitState, FResolvedNetRefHandlesArray& OutNewResolvedRefHandles, const FObjectReferenceTracker& NewUnresolvedReferences, const FObjectReferenceTracker& NewMappedDynamicReferences);
 
 	// Remove all references for object
