Fixing an issue where NetTraceInternal's ThreadBuffer was not actually thread_local due to how UE_AUTORTFM_DECLARE_THREAD_LOCAL_VAR should not be used in static global space. Moving this to function-scope fixes the issue.

#rnx
#tests Tested locally against Dedicated Server and Client
#Synced-CL 48448085

diff --git a/Engine/Source/Runtime/Net/Core/Private/Net/Core/Trace/NetTraceInternal.cpp b/Engine/Source/Runtime/Net/Core/Private/Net/Core/Trace/NetTraceInternal.cpp
--- a/Engine/Source/Runtime/Net/Core/Private/Net/Core/Trace/NetTraceInternal.cpp
+++ b/Engine/Source/Runtime/Net/Core/Private/Net/Core/Trace/NetTraceInternal.cpp
@@ -90,6 +90,7 @@ struct FNetTraceInternal
 
 	static uint32 GetCurrentTraceId();
 
+	static TUniquePtr<FNetTraceInternal::FThreadBuffer>& GetThreadBufferUniquePtr();
 	static inline FThreadBuffer& GetThreadBuffer();
 
 	static constexpr ENetTraceVersion NetTraceVersion = ENetTraceVersion::ENetTraceVersion_DebugNameIndexIs32Bits;
@@ -97,7 +98,6 @@ struct FNetTraceInternal
 
 static std::atomic<uint32> CurrentTraceId(0);
 static std::atomic<UE::Net::FNetDebugNameId> NextNameId(1);
-static UE_AUTORTFM_DECLARE_THREAD_LOCAL_VAR(TUniquePtr<FNetTraceInternal::FThreadBuffer>, ThreadBuffer);
 static FDelegateHandle OnTraceConnectedHandle;
 static FDelegateHandle OnTraceStoppedHandle;
 static FTransactionallySafeCriticalSection NewObjectsCS;
@@ -135,8 +135,15 @@ static FGameInstanceContainer GameInstanceContainer;
 FNetTrace::FOnResetPersistentNetDebugNames FNetTrace::OnResetPersistentNetDebugNames;
 FNetTrace::FOnNetTraceStarted FNetTrace::OnNetTraceStarted;
 
+TUniquePtr<FNetTraceInternal::FThreadBuffer>& FNetTraceInternal::GetThreadBufferUniquePtr()
+{
+	UE_AUTORTFM_DECLARE_THREAD_LOCAL_VAR(TUniquePtr<FNetTraceInternal::FThreadBuffer>, ThreadBuffer);
+	return ThreadBuffer;
+}
+
 FNetTraceInternal::FThreadBuffer& FNetTraceInternal::GetThreadBuffer()
 {
+	TUniquePtr<FNetTraceInternal::FThreadBuffer>& ThreadBuffer = GetThreadBufferUniquePtr();
 	FNetTraceInternal::FThreadBuffer* ThreadBufferPtr = ThreadBuffer.Get();
 	uint32 TraceId = GetCurrentTraceId();
 	if (!ThreadBufferPtr || ThreadBufferPtr->TraceId != TraceId)
@@ -198,9 +205,10 @@ void FNetTrace::OnTraceConnected()
 
 void FNetTrace::OnTraceStopped(FTraceAuxiliary::EConnectionType /*TraceType*/, const FString& /*TraceDestination*/)
 {
-	if (ThreadBuffer)
+	TUniquePtr<FNetTraceInternal::FThreadBuffer>& ThreadBufferUniquePtr = FNetTraceInternal::GetThreadBufferUniquePtr();
+	if (ThreadBufferUniquePtr)
 	{
-		ThreadBuffer.Reset();
+		ThreadBufferUniquePtr.Reset();
 	}
 }
 
