Iris - Fix metrics names and use ASCII strings for initialization of them.

#rnx
#Synced-CL 48600773

diff --git a/Engine/Source/Runtime/Net/Iris/Private/Iris/ReplicationSystem/ReplicationSystem.cpp b/Engine/Source/Runtime/Net/Iris/Private/Iris/ReplicationSystem/ReplicationSystem.cpp
--- a/Engine/Source/Runtime/Net/Iris/Private/Iris/ReplicationSystem/ReplicationSystem.cpp
+++ b/Engine/Source/Runtime/Net/Iris/Private/Iris/ReplicationSystem/ReplicationSystem.cpp
@@ -83,14 +83,14 @@ public:
 	FNetBitArray ConnectionsPendingPostTickDispatchSend;
 	EReplicationSystemSendPass CurrentSendPass = EReplicationSystemSendPass::Invalid;
 
-	FName MetricNameTotalRootObjects;
-	FName MetricNameTotalSubObjects;
-	FName MetricNameAvgPendingObjectCount;
-	FName MetricNameMaxPendingObjectCount;
-	FName MetricNameAvgPendingDependentObjectCount;
-	FName MetricNameMaxPendingDependentObjectCount;
-	FName MetricNameAvgHugeObjectSendQueue;
-	FName MetricNameMaxHugeObjectSendQueue;
+	inline static const FName MetricNameTotalRootObjects = "TotalRootObjects";
+	inline static const FName MetricNameTotalSubObjects = "TotalSubObjects";
+	inline static const FName MetricNameAvgPendingObjectCount = "AvgPendingObjectCount";
+	inline static const FName MetricNameMaxPendingObjectCount = "MaxPendingObjectCount";
+	inline static const FName MetricNameAvgPendingDependentObjectCount= "AvgPendingDependentObjectCount";
+	inline static const FName MetricNameMaxPendingDependentObjectCount = "MaxPendingDependentObjectCount";
+	inline static const FName MetricNameAvgHugeObjectSendQueue = "AvgHugeObjectSendQueue";
+	inline static const FName MetricNameMaxHugeObjectSendQueue = "MaxHugeObjectSendQueue";
 	bool bHadValidConnectionsLastUpdate = false;
 
 	explicit FReplicationSystemImpl(UReplicationSystem* InReplicationSystem, const UReplicationSystem::FReplicationSystemParams& Params)
@@ -111,15 +111,6 @@ public:
 	{
 		const bool bIsProxyBackend = (Params.ProxyType == UE::Net::EProxyType::Backend);
 		ensureMsgf(!bIsProxyBackend || (bIsProxyBackend && !InReplicationSystem->IsServer()), TEXT("A replication system configured as a proxy backend is expected to be a client connecting to a remote server."));
-
-		MetricNameTotalRootObjects = TEXT("TotalSubObjects");
-		MetricNameTotalSubObjects = TEXT("TotalRootObjects");
-		MetricNameAvgPendingObjectCount = TEXT("AvgPendingObjectCount");
-		MetricNameMaxPendingObjectCount = TEXT("MaxPendingObjectCount");
-		MetricNameAvgPendingDependentObjectCount = TEXT("AvgPendingDependentObjectCount");
-		MetricNameMaxPendingDependentObjectCount = TEXT("MaxPendingDependentObjectCount");
-		MetricNameAvgHugeObjectSendQueue = TEXT("AvgHugeObjectSendQueue");
-		MetricNameMaxHugeObjectSendQueue = TEXT("MaxHugeObjectSendQueue");
 	}
 
 	~FReplicationSystemImpl()
