Iris
* Removed more unused code from ReplicationReader
* Deleted ReadObjectsAndSubObjectsPendingDestroy (not called anywhere)
* Remove bShouldProcessHierarchy flag from EnqueueEndReplication (it was always true)
* Slight refactor + more ensures added to QueuedChunk serialization.

#rnx

#Synced-CL 48467961

diff --git a/Engine/Source/Runtime/Net/Iris/Private/Iris/ReplicationSystem/ReplicationReader.cpp b/Engine/Source/Runtime/Net/Iris/Private/Iris/ReplicationSystem/ReplicationReader.cpp
--- a/Engine/Source/Runtime/Net/Iris/Private/Iris/ReplicationSystem/ReplicationReader.cpp
+++ b/Engine/Source/Runtime/Net/Iris/Private/Iris/ReplicationSystem/ReplicationReader.cpp
@@ -462,86 +462,6 @@ uint32 FReplicationReader::ReadObjectsPendingDestroy(FNetSerializationContext& C
 	return ReadRootObjectsPendingDestroy(Context);
 }
 
-uint32 FReplicationReader::ReadObjectsAndSubObjectsPendingDestroy(FNetSerializationContext& Context)
-{
-	FNetBitStreamReader& Reader = *Context.GetBitStreamReader();
-
-	UE_NET_TRACE_SCOPE(ObjectsAndSubObjectsPendingDestroy, Reader, Context.GetTraceCollector(), ENetTraceVerbosity::Trace);
-
-	// Read how many destroyed objects we have
-	constexpr uint32 DestroyObjectBitCount = 16U;
-	constexpr uint32 MaxDestroyObjectCount = (1U << DestroyObjectBitCount) - 1U;
-	const uint32 ObjectsToRead = Reader.ReadBits(DestroyObjectBitCount);
-	
-	if (Context.HasErrorOrOverflow())
-	{
-		return 0;
-	}
-
-	for (uint32 It = 0; It < ObjectsToRead; ++It)
-	{
-		UE_NET_TRACE_NAMED_OBJECT_SCOPE(DestroyedObjectScope, FNetRefHandle::GetInvalid(), Reader, Context.GetTraceCollector(), ENetTraceVerbosity::Trace);
-
-		FNetRefHandle IncompleteHandle = ReadNetRefHandleId(Context, Reader);
-		FNetRefHandle SubObjectRootOrHandle = IncompleteHandle;
-		if (Reader.ReadBool())
-		{
-			SubObjectRootOrHandle = ReadNetRefHandleId(Context, Reader);
-		}
-		const bool bShouldDestroyInstance = Reader.ReadBool();
-		if (Context.HasErrorOrOverflow())
-		{
-			break;
-		}
-
-		ECloseReplicationReason CloseReason = ECloseReplicationReason::None;
-		if (!ReadCloseReason(Context, CloseReason))
-		{
-			break;
-		}
-
-		if (FPendingBatchData* PendingBatchData = PendingBatchHolder.Find(SubObjectRootOrHandle))
-		{
-			// We expect end replication requests for subobjects too so we should only handle one object at a time.
-			constexpr bool bShouldProcessHierarchy = false;
-			EnqueueEndReplication(PendingBatchData, bShouldDestroyInstance, bShouldProcessHierarchy, IncompleteHandle, CloseReason);
-			continue;
-		}
-
-		{
-			// Resolve handle and destroy using bridge
-			const uint32 InternalIndex = NetRefHandleManager->GetInternalIndex(IncompleteHandle);
-			if (InternalIndex != FNetRefHandleManager::InvalidInternalIndex)
-			{
-				UE_NET_TRACE_SET_SCOPE_OBJECTID(DestroyedObjectScope, IncompleteHandle);
-
-				// Defer EndReplication until after applying state data
-				FDispatchObjectInfo& Info = ObjectsToDispatchArray->AddPendingDispatchObjectInfo(TempLinearAllocator);
-
-				Info.bDestroy = bShouldDestroyInstance;
-				Info.bTearOff = false;
-				Info.bDeferredEndReplication = true;
-				Info.InternalIndex = InternalIndex;
-				Info.bIsInitialState = false;
-				Info.bHasState = false;
-				Info.bHasAttachments = false;
-				Info.bShouldCallSubObjectCreatedFromReplication = false;
-				Info.CloseReason = CloseReason;
-
-				// Mark for dispatch
-				ObjectsToDispatchArray->CommitPendingDispatchObjectInfo();
-			}
-			else
-			{
-				// If we did not find the object or associated bridge, the packet that would have created the object may have been lost.
-				UE_LOG(LogIris, Verbose, TEXT("FReplicationReader::Read Tried to destroy object %s. This can occur if the server sends destroy for an object that has not yet been confirmed as created."), *IncompleteHandle.ToString());
-			}
-		}
-	}
-
-	return ObjectsToRead;
-}
-
 uint32 FReplicationReader::ReadRootObjectsPendingDestroy(FNetSerializationContext& Context)
 {
 	FNetBitStreamReader& Reader = *Context.GetBitStreamReader();
@@ -577,8 +497,7 @@ uint32 FReplicationReader::ReadRootObjectsPendingDestroy(FNetSerializationContex
 
 		if (FPendingBatchData* PendingBatchData = PendingBatchHolder.Find(IncompleteHandle))
 		{
-			constexpr bool bShouldProcessHierarchy = true;
-			EnqueueEndReplication(PendingBatchData, bShouldDestroyInstance, bShouldProcessHierarchy, IncompleteHandle, CloseReason);
+			EnqueueEndReplication(PendingBatchData, bShouldDestroyInstance, IncompleteHandle, CloseReason);
 			continue;
 		}
 
@@ -1190,6 +1109,8 @@ uint32 FReplicationReader::ReadObjectBatch(FNetSerializationContext& Context, ui
 		const uint32 NumDataWords = (NumDataBits + 31U) / 32U;
 
 		DataChunk.NumBits = NumDataBits;
+		ensureMsgf((uint32)DataChunk.NumBits == NumDataBits, TEXT("NumDataBits: %u does not fit inside FQueuedDataChunk::NumBits bitfield"), NumDataBits);
+
 		DataChunk.StorageOffset = PendingBatchData->DataChunkStorage.Num();
 		DataChunk.bHasBatchOwnerData = bHasBatchOwnerData;
 		DataChunk.bIsEndReplicationChunk = false;
@@ -2561,12 +2482,22 @@ void FReplicationReader::ProcessHugeObjectAttachment(FNetSerializationContext& C
 	}
 }
 
-bool FReplicationReader::EnqueueEndReplication(FPendingBatchData* PendingBatchData, bool bShouldDestroyInstance, bool bShouldProcessHierarchy, FNetRefHandle NetRefHandleToEndReplication, ECloseReplicationReason CloseReason)
+bool FReplicationReader::EnqueueEndReplication(FPendingBatchData* PendingBatchData, bool bShouldDestroyInstance, FNetRefHandle NetRefHandleToEndReplication, ECloseReplicationReason CloseReason)
 {
 	UE_LOG(LogIris, Verbose, TEXT("FReplicationReader::EnqueueEndReplication for %s since %s has queued batches"), *NetRefHandleToEndReplication.ToString(), *PendingBatchData->Owner.ToString());
-		
-	const uint32 MaxNumDataBits = 67U + static_cast<uint32>(ECloseReplicationReason::MaxBits);
-	const uint32 NumDataWords = (MaxNumDataBits + 31U) / 32U;
+
+	/* Template of the extra data we serialize in the chunk.
+	struct FEndReplicationData
+	{
+		FNetRefHandle EndReplicationHandle;
+		bool bShouldDestroyInstance : 1;
+		ECloseReplicationReason Reason : ECloseReplicationReason::MaxBits;
+	}
+	*/
+	constexpr uint32 MaxNumDataBits = (sizeof(FNetRefHandle) * 8U) + 1U /*bShouldDestroyInstance*/ + (uint32)ECloseReplicationReason::MaxBits;
+	constexpr uint32 NumDataWords = (MaxNumDataBits + 31U) / 32U;
+
+	static_assert(MaxNumDataBits < (1U << 30 /*bitsize of FQueuedDataChunk::NumBits*/), "Bits needed don't fit inside FQueuedDataChunk::NumBits");
 
 	// Enqueue BatchData
 	FQueuedDataChunk DataChunk;
@@ -2586,19 +2517,12 @@ bool FReplicationReader::EnqueueEndReplication(FPendingBatchData* PendingBatchDa
 	WriteUint64(&Writer, NetRefHandleToEndReplication.GetId());
 	Writer.WriteBool(bShouldDestroyInstance);
 	
-	const bool bHasDestroyReason = CloseReason != ECloseReplicationReason::None;
-	Writer.WriteBool(bHasDestroyReason);
-	if (bHasDestroyReason)
-	{
-		const uint32 MaxBits = static_cast<uint32>(ECloseReplicationReason::MaxBits);
-		Writer.WriteBits(static_cast<uint32>(CloseReason), MaxBits);
-	}
+	const uint32 MaxBits = static_cast<uint32>(ECloseReplicationReason::MaxBits);
+	Writer.WriteBits(static_cast<uint32>(CloseReason), MaxBits);
 	
-	Writer.WriteBool(bShouldProcessHierarchy);
-
 	Writer.CommitWrites();
 	ensure(Writer.GetPosBits() <= MaxNumDataBits);
-	
+
 	// Set chunk's actual bit count
 	DataChunk.NumBits = Writer.GetPosBits();
 
@@ -2675,38 +2599,43 @@ void FReplicationReader::ProcessQueuedBatches()
 			{
 				// Read data stored for objects ending replication, 
 				// this can be the batch root or a subobject owned by the batched root.
+
+				/* Template of the extra data we serialized in the chunk.
+				struct FEndReplicationData
+				{
+					FNetRefHandle EndReplicationHandle;
+					bool bShouldDestroyInstance : 1;
+					ECloseReplicationReason CloseReason : ECloseReplicationReason::MaxBits;
+				}
+				*/
+
 				const uint64 NetRefHandleIdToEndReplication = ReadUint64(&Reader);
+				const FNetRefHandle NetRefHandleToEndReplication = FNetRefHandleManager::MakeNetRefHandleFromId(NetRefHandleIdToEndReplication);
+
 				const bool bShouldDestroyInstance = Reader.ReadBool();
 
-				ECloseReplicationReason CloseReason = ECloseReplicationReason::None;
-				if (!ReadCloseReason(Context, CloseReason))
+				const uint32 MaxReasonBits = (uint32)ECloseReplicationReason::MaxBits;
+				ECloseReplicationReason CloseReason = (ECloseReplicationReason)Reader.ReadBits(MaxReasonBits);
+
+				if (Reader.IsOverflown())
 				{
+					ensureMsgf(false, TEXT("Error while reading end replication data. Should not happen"));
 					break;
 				}
-
-				const bool bShouldProcessHierarchy = Reader.ReadBool();
-					
-				const FNetRefHandle NetRefHandleToEndReplication = FNetRefHandleManager::MakeNetRefHandleFromId(NetRefHandleIdToEndReplication);
-					
+				
 				// End replication for object
 				const uint32 InternalIndex = NetRefHandleManager->GetInternalIndex(NetRefHandleToEndReplication);
 				if (InternalIndex != FNetRefHandleManager::InvalidInternalIndex)
 				{
 					constexpr bool bTearOff = false;
-					if (bShouldProcessHierarchy)
-					{
-						EndReplication(InternalIndex, bTearOff, bShouldDestroyInstance, CloseReason);
-						for (const FInternalNetRefIndex SubObjectIndex : NetRefHandleManager->GetSubObjects(InternalIndex))
-						{
-							// End replication for all subobjects and destroy dynamic ones
-							const FNetRefHandle SubObjectRefHandle = NetRefHandleManager->GetNetRefHandleFromInternalIndex(SubObjectIndex);
-							const bool bShouldDestroySubObject = SubObjectRefHandle.IsDynamic();
-							EndReplication(SubObjectIndex, bTearOff, bShouldDestroySubObject, CloseReason);
-						}
-					}
-					else
+
+					EndReplication(InternalIndex, bTearOff, bShouldDestroyInstance, CloseReason);
+					for (const FInternalNetRefIndex SubObjectIndex : NetRefHandleManager->GetSubObjects(InternalIndex))
 					{
-						EndReplication(InternalIndex, bTearOff, bShouldDestroyInstance, CloseReason);
+						// End replication for all subobjects and destroy dynamic ones
+						const FNetRefHandle SubObjectRefHandle = NetRefHandleManager->GetNetRefHandleFromInternalIndex(SubObjectIndex);
+						const bool bShouldDestroySubObject = SubObjectRefHandle.IsDynamic();
+						EndReplication(SubObjectIndex, bTearOff, bShouldDestroySubObject, CloseReason);
 					}
 				}
 
diff --git a/Engine/Source/Runtime/Net/Iris/Private/Iris/ReplicationSystem/ReplicationReader.h b/Engine/Source/Runtime/Net/Iris/Private/Iris/ReplicationSystem/ReplicationReader.h
--- a/Engine/Source/Runtime/Net/Iris/Private/Iris/ReplicationSystem/ReplicationReader.h
+++ b/Engine/Source/Runtime/Net/Iris/Private/Iris/ReplicationSystem/ReplicationReader.h
@@ -155,8 +155,7 @@ private:
 	void ReadStreamDebugFeatures(FNetSerializationContext& Context);
 	// Read objects pending destroy. This will call ReadObjectsAndSubObjectsPendingDestroy or ReadRootObjectsPendingDestroy depending on sending side settings.
 	uint32 ReadObjectsPendingDestroy(FNetSerializationContext& Context);
-	// Receive both objects and subobject destruction. Legacy Iris behavior.
-	uint32 ReadObjectsAndSubObjectsPendingDestroy(FNetSerializationContext& Context);
+
 	// Receive only root object destruction and destroy all existing subobjects with them atomically.
 	uint32 ReadRootObjectsPendingDestroy(FNetSerializationContext& Context);
 
@@ -223,7 +222,7 @@ private:
 
 	// If we are queuing data for a batch we must also defer calls to EndReplication
 	// This method writes this method in the form of a QueuedChunk
-	bool EnqueueEndReplication(FPendingBatchData* PendingBatchData, bool bShouldDestroyInstance, bool bShouldProcessHierarchy, FNetRefHandle NetRefHandleToEndReplication, ECloseReplicationReason CloseReason);
+	bool EnqueueEndReplication(FPendingBatchData* PendingBatchData, bool bShouldDestroyInstance, FNetRefHandle NetRefHandleToEndReplication, ECloseReplicationReason CloseReason);
 
 	// Remove a handle from the hot and cold unresolved caches used by ResolveAndDispatchUnresolvedReferences(). If this handle is marked as unresolved
 	// again, it will be added to the hot cache.
