Iris - Fix race condition if a child NetStatsContext needs to be allocated for parallel processing.

#13704
#Synced-CL 48600773

diff --git a/Engine/Source/Runtime/Net/Iris/Private/Iris/Stats/NetStats.cpp b/Engine/Source/Runtime/Net/Iris/Private/Iris/Stats/NetStats.cpp
--- a/Engine/Source/Runtime/Net/Iris/Private/Iris/Stats/NetStats.cpp
+++ b/Engine/Source/Runtime/Net/Iris/Private/Iris/Stats/NetStats.cpp
@@ -209,32 +209,26 @@ FNetStatsContext* FNetTypeStats::CreateChildNetStatsContext()
 
 FNetStatsContext* FNetTypeStats::AcquireChildNetStatsContext()
 {
-	int32 AvailableIndex = -1;
+	FScopeLock Lock(&ChildStatsContextCS);
 	
-	ChildStatsContextCS.Lock();
-	for (int32 SearchIndex = 0; SearchIndex < ChildStatsContext.Num() && AvailableIndex == -1; SearchIndex++)
+	for (FNetStatsContext* StatsContext : ChildStatsContext)
 	{
-		if (ChildStatsContext[SearchIndex]->CanAcquire())
+		if (StatsContext->CanAcquire())
 		{
-			AvailableIndex = SearchIndex;
-			ChildStatsContext[SearchIndex]->Acquire();
+			StatsContext->Acquire();
+			return StatsContext;
 		}
 	}
-	ChildStatsContextCS.Unlock();
 
 	//Fallback to create a new context if we've run out (PreUpdateSetup should have created enough)
-	ensureMsgf(AvailableIndex != -1, TEXT("Too many tasks accessing ChildNetStatsContext, creating a new one. This may cause unexpected memory usage. Current count: %i"), ChildStatsContext.Num());
-	if(AvailableIndex == -1)
+	ensureMsgf(false, TEXT("Too many tasks accessing ChildNetStatsContext, creating a new one. This may cause unexpected memory usage. Current count: %i"), ChildStatsContext.Num());
+	FNetStatsContext* NewContext = CreateChildNetStatsContext();
+	if (ensure(NewContext))
 	{
-		FNetStatsContext* NewContext = CreateChildNetStatsContext();
-		check(NewContext);
-		ChildStatsContextCS.Lock();
 		NewContext->Acquire();
-		ChildStatsContextCS.Unlock();
-		return NewContext;
 	}
 
-	return ChildStatsContext[AvailableIndex];
+	return NewContext;
 }
 
 void FNetTypeStats::ReleaseChildNetStatsContext(FNetStatsContext* StatsContext)
