Iris - Fix compile error in NetStats.cpp

#rnx
#tests Program previously exhibiting compile error builds without issues.
#Synced-CL 48282875

diff --git a/Engine/Source/Runtime/Net/Iris/Private/Iris/Stats/NetStats.cpp b/Engine/Source/Runtime/Net/Iris/Private/Iris/Stats/NetStats.cpp
--- a/Engine/Source/Runtime/Net/Iris/Private/Iris/Stats/NetStats.cpp
+++ b/Engine/Source/Runtime/Net/Iris/Private/Iris/Stats/NetStats.cpp
@@ -209,7 +209,7 @@ FNetStatsContext* FNetTypeStats::CreateChildNetStatsContext()
 
 FNetStatsContext* FNetTypeStats::AcquireChildNetStatsContext()
 {
-	FScopeLock Lock(&ChildStatsContextCS);
+	UE::TScopeLock Lock(ChildStatsContextCS);
 	
 	for (FNetStatsContext* StatsContext : ChildStatsContext)
 	{
