Iris
* More rewriting and cleanup in ReplicationReader
* Made EReplicatedDestroyHeaderFlags a proper enum class
* Removed cvar net.Iris.ImmediateDispatchEndReplicationForSubObjects : only true is supported now
* Removed cvar net.Iris.AllowSendStateWithEndReplicationForSubObjects: only false is supported now

#rnx
#Synced-CL 48467961

diff --git a/Engine/Plugins/Runtime/ReplicationSystemTestPlugin/Source/Private/Tests/ReplicationSystem/TestReplicationSystem.cpp b/Engine/Plugins/Runtime/ReplicationSystemTestPlugin/Source/Private/Tests/ReplicationSystem/TestReplicationSystem.cpp
--- a/Engine/Plugins/Runtime/ReplicationSystemTestPlugin/Source/Private/Tests/ReplicationSystem/TestReplicationSystem.cpp
+++ b/Engine/Plugins/Runtime/ReplicationSystemTestPlugin/Source/Private/Tests/ReplicationSystem/TestReplicationSystem.cpp
@@ -2407,18 +2407,6 @@ UE_NET_TEST_FIXTURE(FReplicationSystemServerClientTestFixture, TestReplicationRe
 
 UE_NET_TEST_FIXTURE(FReplicationSystemServerClientTestFixture, TestSubObjectReplicatedDestroyBeforePostNetReceive)
 {
-	// Make sure that bOldImmediateDispatchEndReplicationForSubObjects is set to true
-	IConsoleVariable* LocalCVar = IConsoleManager::Get().FindConsoleVariable(TEXT("net.Iris.ImmediateDispatchEndReplicationForSubObjects"));
-	check(LocalCVar != nullptr && LocalCVar->IsVariableBool());
-	const bool bOldImmediateDispatchEndReplicationForSubObjects = LocalCVar->GetBool();
-	LocalCVar->Set(true, ECVF_SetByCode);
-
-	ON_SCOPE_EXIT
-	{
-		// Restore cvars
-		LocalCVar->Set(bOldImmediateDispatchEndReplicationForSubObjects, ECVF_SetByCode);
-	};
-
 	// Add a client
 	FReplicationSystemTestClient* Client = CreateClient();
 
@@ -2465,66 +2453,6 @@ UE_NET_TEST_FIXTURE(FReplicationSystemServerClientTestFixture, TestSubObjectRepl
 	UE_NET_ASSERT_FALSE(ClientSubObject2->ObjectToWatch->IsValid());
 }
 
-UE_NET_TEST_FIXTURE(FReplicationSystemServerClientTestFixture, TestSubObjectReplicatedDestroyAfterPostNetReceive)
-{
-	// Make sure that bOldImmediateDispatchEndReplicationForSubObjects is set to false
-	IConsoleVariable* LocalCVar = IConsoleManager::Get().FindConsoleVariable(TEXT("net.Iris.ImmediateDispatchEndReplicationForSubObjects"));
-	check(LocalCVar != nullptr && LocalCVar->IsVariableBool());
-	const bool bOldImmediateDispatchEndReplicationForSubObjects = LocalCVar->GetBool();
-	LocalCVar->Set(false, ECVF_SetByCode);
-
-	ON_SCOPE_EXIT
-	{
-		// Restore cvars
-		LocalCVar->Set(bOldImmediateDispatchEndReplicationForSubObjects, ECVF_SetByCode);
-	};
-
-	// Add a client
-	FReplicationSystemTestClient* Client = CreateClient();
-
-	// Spawn object on server
-	UReplicatedTestObject* ServerObject = Server->CreateObject(0, 0);
-
-	// Spawn some subobjects
-	UReplicatedSubObjectDestroyOrderObject* ServerSubObject0 = Server->CreateSubObject<UReplicatedSubObjectDestroyOrderObject>(ServerObject->NetRefHandle);
-	UReplicatedSubObjectDestroyOrderObject* ServerSubObject1 = Server->CreateSubObject<UReplicatedSubObjectDestroyOrderObject>(ServerObject->NetRefHandle);
-	UReplicatedSubObjectDestroyOrderObject* ServerSubObject2 = Server->CreateSubObject<UReplicatedSubObjectDestroyOrderObject>(ServerObject->NetRefHandle);
-
-	UReplicatedSubObjectDestroyOrderObject::RepOrderCounter = 0U;
-
-	// Send and deliver packet
-	Server->UpdateAndSend({Client});
-
-	// Verify that objects have replicated
-	UReplicatedSubObjectDestroyOrderObject* ClientSubObject0 = Cast<UReplicatedSubObjectDestroyOrderObject>(Client->GetReplicationBridge()->GetReplicatedObject(ServerSubObject0->NetRefHandle));
-	UReplicatedSubObjectDestroyOrderObject* ClientSubObject1 = Cast<UReplicatedSubObjectDestroyOrderObject>(Client->GetReplicationBridge()->GetReplicatedObject(ServerSubObject1->NetRefHandle));
-	UReplicatedSubObjectDestroyOrderObject* ClientSubObject2 = Cast<UReplicatedSubObjectDestroyOrderObject>(Client->GetReplicationBridge()->GetReplicatedObject(ServerSubObject2->NetRefHandle));
-
-	UE_NET_ASSERT_NE(ClientSubObject0, nullptr);
-	UE_NET_ASSERT_NE(ClientSubObject1, nullptr);
-	UE_NET_ASSERT_NE(ClientSubObject2, nullptr);
-
-	// Verify that they have replicated in expected order
-	UE_NET_ASSERT_EQ(ClientSubObject0->LastRepOrderCounter, 1U);
-	UE_NET_ASSERT_GT(ClientSubObject1->LastRepOrderCounter, ClientSubObject0->LastRepOrderCounter);
-	UE_NET_ASSERT_GT(ClientSubObject2->LastRepOrderCounter, ClientSubObject1->LastRepOrderCounter);	
-
-	// setup a watch on the client
-	ClientSubObject2->SetObjectExpectedToBeDestroyed(ClientSubObject1);
-
-	// Dirty some data on server and destroy SubObject1
-	++ServerSubObject2->IntA;
-	Server->ReplicationBridge->EndReplication(ServerSubObject1, EEndReplicationFlags::Destroy);
-
-	// Send and deliver packet
-	Server->UpdateAndSend({Client});
-
-	// Verify assumptions, replicated subobject destroy should have been issued after ClientSubObjects2 s call to PostNetReceive
-	UE_NET_ASSERT_TRUE(ClientSubObject2->bObjectExistedInPreNetReceive);
-	UE_NET_ASSERT_TRUE(ClientSubObject2->bObjectExistedInPostNetReceive);
-	UE_NET_ASSERT_FALSE(ClientSubObject2->ObjectToWatch->IsValid());
-}
-
 } // end namespace UE::Net::Private
 
 
diff --git a/Engine/Source/Runtime/Net/Iris/Private/Iris/ReplicationSystem/ReplicationReader.cpp b/Engine/Source/Runtime/Net/Iris/Private/Iris/ReplicationSystem/ReplicationReader.cpp
--- a/Engine/Source/Runtime/Net/Iris/Private/Iris/ReplicationSystem/ReplicationReader.cpp
+++ b/Engine/Source/Runtime/Net/Iris/Private/Iris/ReplicationSystem/ReplicationReader.cpp
@@ -81,12 +81,6 @@ static FAutoConsoleVariableRef CVarExecuteReliableRPCsBeforeApplyState(
 		bExecuteReliableRPCsBeforeApplyState,
 		TEXT("If true and Iris runs in backwards compatibility mode then reliable RPCs will be executed before we apply state data on the target object unless we first need to spawn the object."));
 
-static bool bImmediateDispatchEndReplicationForSubObjects = true;
-static FAutoConsoleVariableRef CVarImmediateDispatchEndReplicationForSubObjects(
-	TEXT("net.Iris.ImmediateDispatchEndReplicationForSubObjects"),
-	bImmediateDispatchEndReplicationForSubObjects,
-	TEXT("if true replicated calls to EndReplication for SubObjects will dispatched immediately. Default is true."));
-
 static bool bDispatchUnresolvedPreviouslyReceivedChanges = false;
 static FAutoConsoleVariableRef CvarDispatchUnresolvedPreviouslyReceivedChanges(
 	TEXT("net.Iris.DispatchUnresolvedPreviouslyReceivedChanges"),
@@ -209,10 +203,19 @@ struct FReplicationReader::FDispatchObjectInfo
 	bool bHasAttachments : 1 = false;
 	bool bDestroy : 1 = false;
 	bool bTearOff : 1 = false;
-	bool bDeferredEndReplication : 1 = false;
-	bool bReplicatedDestroyForSubObject : 1 = false;
-	bool bProcessEndReplicationForHierarchy : 1 = false;
+	
+	/** When true we will call EndReplication on the object after reading all the bunch */
+	bool bWantsDeferredEndReplication : 1 = false;
+	
+	/** When true we will immediately call EndReplication on the object while still reading data instead of waiting for the entire bunch to be read */
+	bool bWantsImmediateEndReplication : 1 = false;
+
+	/** When true we will call EndReplication on all it's subobjects */
+	bool bCallEndReplicationOnSubObjects : 1 = false;
+
+	/** When true we'll tell the RootObject's factory that this subobject was instantiated */
 	bool bShouldCallSubObjectCreatedFromReplication : 1 = false;
+
 	ECloseReplicationReason CloseReason = ECloseReplicationReason::None;
 };
 
@@ -513,8 +516,8 @@ uint32 FReplicationReader::ReadRootObjectsPendingDestroy(FNetSerializationContex
 
 				Info.bDestroy = bShouldDestroyInstance;
 				Info.bTearOff = false;
-				Info.bDeferredEndReplication = true;
-				Info.bProcessEndReplicationForHierarchy = true;
+				Info.bWantsDeferredEndReplication = true;
+				Info.bCallEndReplicationOnSubObjects = true;
 				Info.InternalIndex = InternalIndex;
 				Info.bIsInitialState = false;
 				Info.bHasState = false;
@@ -591,36 +594,40 @@ void FReplicationReader::CleanupObjectData(FReplicatedObjectInfo& ObjectInfo)
 
 void FReplicationReader::EndReplication(FInternalNetRefIndex InternalIndex, bool bTearOff, bool bDestroyInstance, ECloseReplicationReason CloseReason)
 {
-	if (!ensure(InternalIndex != ObjectIndexForOOBAttachment))
+	if (InternalIndex == ObjectIndexForOOBAttachment)
 	{
+		ensure(InternalIndex != ObjectIndexForOOBAttachment);
 		return;
 	}
 	
-	if (FReplicatedObjectInfo* ObjectInfo = ReplicatedObjects.Find(InternalIndex))
+	FReplicatedObjectInfo* ObjectInfo = ReplicatedObjects.Find(InternalIndex);
+	if (!ObjectInfo)
 	{
-		CleanupReferenceTracking(ObjectInfo);
-		Attachments.DropAllAttachments(ENetObjectAttachmentType::Normal, InternalIndex);
+		return;
+	}
 
-		EReplicationBridgeDestroyInstanceReason DestroyReason = EReplicationBridgeDestroyInstanceReason::DoNotDestroy;
-		if (bTearOff)
-		{
-			DestroyReason = EReplicationBridgeDestroyInstanceReason::TearOff;
-		}
-		else if (CloseReason == ECloseReplicationReason::Migrated)
-		{
-			DestroyReason = EReplicationBridgeDestroyInstanceReason::Migrated;
-		}
-		else if (bDestroyInstance)
-		{
-			DestroyReason = EReplicationBridgeDestroyInstanceReason::Destroy;
-		}
+	CleanupReferenceTracking(ObjectInfo);
+	Attachments.DropAllAttachments(ENetObjectAttachmentType::Normal, InternalIndex);
+
+	EReplicationBridgeDestroyInstanceReason DestroyReason = EReplicationBridgeDestroyInstanceReason::DoNotDestroy;
+	if (bTearOff)
+	{
+		DestroyReason = EReplicationBridgeDestroyInstanceReason::TearOff;
+	}
+	else if (CloseReason == ECloseReplicationReason::Migrated)
+	{
+		DestroyReason = EReplicationBridgeDestroyInstanceReason::Migrated;
+	}
+	else if (bDestroyInstance)
+	{
+		DestroyReason = EReplicationBridgeDestroyInstanceReason::Destroy;
+	}
 		
-		FBridgeSerialization::DestroyNetObjectFromRemote(ReplicationBridge, InternalIndex, DestroyReason);
+	FBridgeSerialization::DestroyNetObjectFromRemote(ReplicationBridge, InternalIndex, DestroyReason);
 
-		CleanupObjectData(*ObjectInfo);
+	CleanupObjectData(*ObjectInfo);
 
-		ReplicatedObjects.Remove(InternalIndex);
-	}
+	ReplicatedObjects.Remove(InternalIndex);
 }
 
 void FReplicationReader::DeserializeObjectStateDelta(FNetSerializationContext& Context, uint32 InternalIndex, FDispatchObjectInfo& Info, FReplicatedObjectInfo& ObjectInfo, const FNetRefHandleManager::FReplicatedObjectData& ObjectData, uint32& OutNewBaselineIndex)
@@ -1182,7 +1189,12 @@ void FReplicationReader::ReadObjectInBatch(FNetSerializationContext& Context, FN
 	
 	// Read replicated destroy header if necessary. We don't know the internal index yet so can't do the more appropriate check IsObjectIndexForOOBAttachment.
 	const bool bReadReplicatedDestroyHeader = IncompleteHandle.IsValid();
-	const uint32 ReplicatedDestroyHeaderFlags = bReadReplicatedDestroyHeader ? Reader.ReadBits(ReplicatedDestroyHeaderFlags_BitCount) : ReplicatedDestroyHeaderFlags_None;
+	EReplicatedDestroyHeaderFlags ReplicatedDestroyHeaderFlags = EReplicatedDestroyHeaderFlags::None;
+	if (bReadReplicatedDestroyHeader)
+	{
+		ReplicatedDestroyHeaderFlags = (EReplicatedDestroyHeaderFlags)Reader.ReadBits(GetDestroyHeaderFlagsBitCount());
+	}
+	 
 
 	const bool bHasState = Reader.ReadBool();
 	if (bHasState)
@@ -1280,6 +1292,8 @@ void FReplicationReader::ReadObjectInBatch(FNetSerializationContext& Context, FN
 
 		InternalIndex = NetRefHandleManager->GetInternalIndex(NetRefHandle);
 		FNetRefHandleManager::FReplicatedObjectData& ObjectData = NetRefHandleManager->GetReplicatedObjectDataNoCheck(InternalIndex);
+
+		// Cache the creation flags given by the factory
 		ObjectData.bAllowDestroyInstanceFromRemote = EnumHasAnyFlags(CreateResult.Flags, EReplicationBridgeCreateNetRefHandleResultFlags::AllowDestroyInstanceFromRemote);
 		bShouldCallSubObjectCreatedFromReplication = EnumHasAnyFlags(CreateResult.Flags, EReplicationBridgeCreateNetRefHandleResultFlags::ShouldCallSubObjectCreatedFromReplication);
 		
@@ -1314,7 +1328,7 @@ void FReplicationReader::ReadObjectInBatch(FNetSerializationContext& Context, FN
 
 			if (InternalIndex == FNetRefHandleManager::InvalidInternalIndex)
 			{
-				if ((ReplicatedDestroyHeaderFlags & ReplicatedDestroyHeaderFlags_EndReplication) == 0U)
+				if (!EnumHasAnyFlags(ReplicatedDestroyHeaderFlags, EReplicatedDestroyHeaderFlags::IsSubObject))
 				{
 					bHasErrors = true;
 					UE_LOG(LogIris, Error, TEXT("FReplicationReader::ReadObject Handle %s not bound to any InternalIndex"), *IncompleteHandle.ToString());
@@ -1344,20 +1358,34 @@ void FReplicationReader::ReadObjectInBatch(FNetSerializationContext& Context, FN
 		FDispatchObjectInfo& Info = ObjectsToDispatchArray->AddPendingDispatchObjectInfo(TempLinearAllocator);
 
 		// Update info based on ReplicatedDestroyHeader
-		Info.bDestroy = !!(ReplicatedDestroyHeaderFlags & (ReplicatedDestroyHeaderFlags_TearOff | ReplicatedDestroyHeaderFlags_DestroyInstance));
-		Info.bTearOff = !!(ReplicatedDestroyHeaderFlags & ReplicatedDestroyHeaderFlags_TearOff);
-		// We set one of these flags
-		if (bImmediateDispatchEndReplicationForSubObjects && Info.bDestroy && !Info.bTearOff)
+		const bool bIsTearOff = EnumHasAnyFlags(ReplicatedDestroyHeaderFlags, EReplicatedDestroyHeaderFlags::TearOff);
+		const bool bIsForSubObject = EnumHasAnyFlags(ReplicatedDestroyHeaderFlags, EReplicatedDestroyHeaderFlags::IsSubObject);
+		const bool bDestroySubObject = EnumHasAnyFlags(ReplicatedDestroyHeaderFlags, EReplicatedDestroyHeaderFlags::DestroySubObject);
+		
+		Info.bDestroy = bDestroySubObject || bIsTearOff;
+		Info.bTearOff = bIsTearOff;
+
+		if (Info.bTearOff)
 		{
-			Info.bReplicatedDestroyForSubObject = true;
-			Info.bDeferredEndReplication = false;
+			Info.bWantsDeferredEndReplication = true;
 		}
-		else
+		else if (bIsForSubObject)
+		{
+			if (bDestroySubObject)
+			{
+				Info.bWantsImmediateEndReplication = true;
+				Info.bWantsDeferredEndReplication = false;
+			}
+			else
+			{
+				Info.bWantsDeferredEndReplication = true;
+			}
+		}
+
+		if (bIsInitialState)
 		{
-			Info.bReplicatedDestroyForSubObject = false;
-			Info.bDeferredEndReplication = !!(ReplicatedDestroyHeaderFlags & (ReplicatedDestroyHeaderFlags_TearOff | ReplicatedDestroyHeaderFlags_EndReplication));
+			Info.bShouldCallSubObjectCreatedFromReplication = bShouldCallSubObjectCreatedFromReplication;
 		}
-		Info.bShouldCallSubObjectCreatedFromReplication = bShouldCallSubObjectCreatedFromReplication;
 
 		if (bHasState)
 		{
@@ -1962,15 +1990,12 @@ void FReplicationReader::DispatchStateData(FNetSerializationContext& Context)
 	{
 		// Dispatch replicated subobject destroy here to behave as close as possible to subobject replication in ActorChannel.
 		// Added to solve issues with assumptions made by blueprint logic when replacing a component with another of the same type
-		if (bImmediateDispatchEndReplicationForSubObjects)
+		for (FPostDispatchObjectInfo& PostDispatchObjectInfo : MakeArrayView(PostDispatchObjectInfos, NumObjectsPendingPostDistpatch))
 		{
-			for (FPostDispatchObjectInfo& PostDispatchObjectInfo : MakeArrayView(PostDispatchObjectInfos, NumObjectsPendingPostDistpatch))
+			FDispatchObjectInfo& Info = *PostDispatchObjectInfo.Info;
+			if (Info.bWantsImmediateEndReplication)
 			{
-				FDispatchObjectInfo& Info = *PostDispatchObjectInfo.Info;
-				if (Info.bReplicatedDestroyForSubObject)
-				{
-					EndReplication(Info.InternalIndex, Info.bTearOff, Info.bDestroy, Info.CloseReason);
-				}
+				EndReplication(Info.InternalIndex, Info.bTearOff, Info.bDestroy, Info.CloseReason);
 			}
 		}
 
@@ -2367,13 +2392,14 @@ void FReplicationReader::UpdateUnresolvableReferenceTracking()
 
 void FReplicationReader::DispatchEndReplication(FNetSerializationContext& Context)
 {
-	for (FDispatchObjectInfo& Info : ObjectsToDispatchArray->GetObjectsToDispatch())
+	for (const FDispatchObjectInfo& Info : ObjectsToDispatchArray->GetObjectsToDispatch())
 	{
-		if (Info.bDeferredEndReplication)
+		if (Info.bWantsDeferredEndReplication)
 		{
 			// Detach and destroy object
 			EndReplication(Info.InternalIndex, Info.bTearOff, Info.bDestroy, Info.CloseReason);
-			if (Info.bProcessEndReplicationForHierarchy)
+
+			if (Info.bCallEndReplicationOnSubObjects)
 			{
 				for (FInternalNetRefIndex SubObjectIndex : NetRefHandleManager->GetSubObjects(Info.InternalIndex))
 				{
@@ -2630,6 +2656,7 @@ void FReplicationReader::ProcessQueuedBatches()
 					constexpr bool bTearOff = false;
 
 					EndReplication(InternalIndex, bTearOff, bShouldDestroyInstance, CloseReason);
+
 					for (const FInternalNetRefIndex SubObjectIndex : NetRefHandleManager->GetSubObjects(InternalIndex))
 					{
 						// End replication for all subobjects and destroy dynamic ones
diff --git a/Engine/Source/Runtime/Net/Iris/Private/Iris/ReplicationSystem/ReplicationTypes.h b/Engine/Source/Runtime/Net/Iris/Private/Iris/ReplicationSystem/ReplicationTypes.h
--- a/Engine/Source/Runtime/Net/Iris/Private/Iris/ReplicationSystem/ReplicationTypes.h
+++ b/Engine/Source/Runtime/Net/Iris/Private/Iris/ReplicationSystem/ReplicationTypes.h
@@ -1,8 +1,13 @@
 // Copyright Epic Games, Inc. All Rights Reserved.
 
 #pragma once
+
+#include "Iris/Core/BitTwiddling.h"
+
 #include "Net/Core/NetBitArray.h"
 
+#include "Misc/EnumClassFlags.h"
+
 class UReplicationSystem;
 
 namespace UE::Net::Private
@@ -29,13 +34,20 @@ struct FReplicationParameters
 	uint32 NumBitsUsedForHugeObjectBatchSize = 32U;
 };
 
-enum EReplicatedDestroyHeaderFlags : uint32
+enum class EReplicatedDestroyHeaderFlags : uint32
 {
-	ReplicatedDestroyHeaderFlags_None						= 0U,
-	ReplicatedDestroyHeaderFlags_TearOff					= 1U << 0U,
-	ReplicatedDestroyHeaderFlags_EndReplication				= ReplicatedDestroyHeaderFlags_TearOff << 1U,
-	ReplicatedDestroyHeaderFlags_DestroyInstance			= ReplicatedDestroyHeaderFlags_EndReplication << 1U,
-	ReplicatedDestroyHeaderFlags_BitCount					= 3U
+	None				= 0U,
+	TearOff				= 1U << 0U,
+	IsSubObject    		= TearOff << 1U,
+	DestroySubObject	= IsSubObject << 1U,
+	BitCount			= 3U //update as necessary
 };
 
+consteval uint32 GetDestroyHeaderFlagsBitCount()
+{
+	return (uint32)EReplicatedDestroyHeaderFlags::BitCount;
+}
+
+ENUM_CLASS_FLAGS(EReplicatedDestroyHeaderFlags);
+
 }
diff --git a/Engine/Source/Runtime/Net/Iris/Private/Iris/ReplicationSystem/ReplicationWriter.cpp b/Engine/Source/Runtime/Net/Iris/Private/Iris/ReplicationSystem/ReplicationWriter.cpp
--- a/Engine/Source/Runtime/Net/Iris/Private/Iris/ReplicationSystem/ReplicationWriter.cpp
+++ b/Engine/Source/Runtime/Net/Iris/Private/Iris/ReplicationSystem/ReplicationWriter.cpp
@@ -108,9 +108,6 @@ static int32 GReplicationWriterCannotSendWarningInterval = 256;
 static FAutoConsoleVariableRef CVarReplicationWriterCannotSendWarningInterval(TEXT("net.Iris.ReplicationWriterCannotSendWarningInterval"), GReplicationWriterCannotSendWarningInterval,
 	TEXT("Warn if we are prevented from Sending for more than a certain number of updates, setting it to 0 will disable warning"));
 
-static bool bAllowSendStateWithEndReplicationForSubObjects = false;
-static FAutoConsoleVariableRef CvarAllowSendStateWithEndReplicationForSubObjects(TEXT("net.Iris.AllowSendStateWithEndReplicationForSubObjects"), bAllowSendStateWithEndReplicationForSubObjects, TEXT("Allow state to be sent together with subobject destroy. Default is false. Cannot be true if net.Iris.ImmediateDispatchEndReplicationForSubObjects is true."));
-
 // Drop unsent ordered unreliable attachments at the end of the tick, 
 // -1 = Allow them to be queued for next tick, 
 // 0  = Always drop at the end of the tick
@@ -249,7 +246,7 @@ static bool s_ValidateReplicationRecord(const FReplicationRecord* ReplicationRec
 		for (uint32 It = 0U; It < RecordInfoCount; ++It)
 		{
  			const FReplicationRecord::FRecordInfo& RecordInfo = ReplicationRecord->PeekInfoAtOffset(It + Offset);
-			// We allow multiple entries for the OOB attachments but do not expect multiple entires for normal replicated objects
+			// We allow multiple entries for the OOB attachments but do not expect multiple entries for normal replicated objects
 			if (RecordInfo.Index != 0U && BitArray.GetBit(RecordInfo.Index))
 			{
 				ensure(false);
@@ -760,13 +757,12 @@ void FReplicationWriter::SetPendingDestroyOrSubObjectPendingDestroyState(uint32
 				MarkObjectDirty(InternalIndex, "SetPendingDestroyOrSubObjectPendingDestroyState2");
 				ObjectsPendingDestroy.SetBit(InternalIndex);
 				Info.SubObjectPendingDestroy = 1U;
-				if (!bAllowSendStateWithEndReplicationForSubObjects)
-				{
-					// For clarity clear HasDirtyChangeMask as we do not intend to replicate state. The object itself needs to remain dirty in ObjectsWithDirtyChanges though.
-					Info.HasDirtyChangeMask = 0U;
-					// One cannot cancel a subobject destroy so we should not send any attachments for this subobject going forward
-					Info.HasAttachments = 0U;
-				}
+			
+				// For clarity clear HasDirtyChangeMask as we do not intend to replicate state. The object itself needs to remain dirty in ObjectsWithDirtyChanges though.
+				Info.HasDirtyChangeMask = 0U;
+				// One cannot cancel a subobject destroy so we should not send any attachments for this subobject going forward
+				Info.HasAttachments = 0U;
+
 				ensure(!Info.TearOff || (GetFlushStatus(InternalIndex, Info, Info.FlushFlags) == FlushFlags_None));
 				return;
 			}
@@ -2232,47 +2228,45 @@ uint32 FReplicationWriter::WriteRootObjectsPendingDestroy(FNetSerializationConte
 			Writer.WriteBits(static_cast<uint32>(ObjectData.CloseReason), MaxBits);
 		}
 
-		if (!Writer.IsOverflown())
+		if (Writer.IsOverflown())
 		{
-			// Must update state before pushing record
-			SetState(InternalIndex, EReplicatedObjectState::WaitOnDestroyConfirmation);
-			Info.HasDirtyChangeMask = 0U;
-
-			// update transmission record
-			FBatchObjectInfo ObjectInfo = {};
-			FObjectRecord ObjectRecord;
-			FReplicationRecord::FSubObjectRecord SubObjectRecord;
+			break;
+		}
 
-			ObjectInfo.InternalIndex = InternalIndex;
-			CreateObjectRecord(nullptr, Info, ObjectInfo, ObjectRecord);
+		// Must update state before pushing record
+		SetState(InternalIndex, EReplicatedObjectState::WaitOnDestroyConfirmation);
+		Info.HasDirtyChangeMask = 0U;
 
-			// Fill in subobject record
-			for (const uint32 SubObjectIndex : NetRefHandleManager->GetSubObjects(InternalIndex))
-			{
-				FReplicationInfo& SubObjectInfo = GetReplicationInfo(SubObjectIndex);
-				if (SubObjectInfo.GetState() == EReplicatedObjectState::Invalid)
-				{
-					continue;
-				}
+		// update transmission record
+		FBatchObjectInfo ObjectInfo = {};
+		FObjectRecord ObjectRecord;
+		FReplicationRecord::FSubObjectRecord SubObjectRecord;
 
-				SetState(SubObjectIndex, EReplicatedObjectState::WaitOnDestroyConfirmation);
-				SubObjectInfo.HasDirtyChangeMask = 0U;
-				SubObjectInfo.SubObjectPendingDestroy = 0U;
+		ObjectInfo.InternalIndex = InternalIndex;
+		CreateObjectRecord(nullptr, Info, ObjectInfo, ObjectRecord);
 
-				FReplicationRecord::FSubObjectRecord::FSubObjectInfo& SubObjectRecordInfo = SubObjectRecord.SubObjectInfos.AddDefaulted_GetRef();
-				SubObjectRecordInfo.Index = SubObjectIndex;
-				SubObjectRecordInfo.ReplicatedObjectState = static_cast<uint32>(SubObjectInfo.GetState());
+		// Fill in subobject record
+		for (const uint32 SubObjectIndex : NetRefHandleManager->GetSubObjects(InternalIndex))
+		{
+			FReplicationInfo& SubObjectInfo = GetReplicationInfo(SubObjectIndex);
+			if (SubObjectInfo.GetState() == EReplicatedObjectState::Invalid)
+			{
+				continue;
 			}
 
-			ObjectRecord.Record.HasSubObjectRecord = !SubObjectRecord.SubObjectInfos.IsEmpty();
-			CommitObjectDestroyRecord(InternalIndex, ObjectRecord, SubObjectRecord);
+			SetState(SubObjectIndex, EReplicatedObjectState::WaitOnDestroyConfirmation);
+			SubObjectInfo.HasDirtyChangeMask = 0U;
+			SubObjectInfo.SubObjectPendingDestroy = 0U;
 
-			++WrittenCount;
-		}
-		else
-		{
-			break;
+			FReplicationRecord::FSubObjectRecord::FSubObjectInfo& SubObjectRecordInfo = SubObjectRecord.SubObjectInfos.AddDefaulted_GetRef();
+			SubObjectRecordInfo.Index = SubObjectIndex;
+			SubObjectRecordInfo.ReplicatedObjectState = static_cast<uint32>(SubObjectInfo.GetState());
 		}
+
+		ObjectRecord.Record.HasSubObjectRecord = !SubObjectRecord.SubObjectInfos.IsEmpty();
+		CommitObjectDestroyRecord(InternalIndex, ObjectRecord, SubObjectRecord);
+
+		++WrittenCount;
 	}
 
 	// Write Header
@@ -2630,7 +2624,7 @@ FReplicationWriter::EWriteObjectStatus FReplicationWriter::WriteObjectAndSubObje
 	}
 	Info.HasDirtyConditionals = 0;
 
-	const bool bAllowStateAndAttachmentSending = (!Info.SubObjectPendingDestroy || bAllowSendStateWithEndReplicationForSubObjects);
+	const bool bAllowStateAndAttachmentSending = !Info.SubObjectPendingDestroy;
 	const bool bIsObjectIndexForAttachment = IsObjectIndexForOOBAttachment(InternalIndex);
 	const bool bHasState = (bIsInitialState || Info.HasDirtyChangeMask) && !!WriteObjectTraits.bWriteState && bAllowStateAndAttachmentSending;
 	const bool bHasAttachments = (Info.HasAttachments || bIsObjectIndexForAttachment) && bAllowStateAndAttachmentSending;
@@ -2671,7 +2665,7 @@ FReplicationWriter::EWriteObjectStatus FReplicationWriter::WriteObjectAndSubObje
 		if (bWriteReplicatedDestroyHeader)
 		{
 			// Write destroy header bits, we always want to write the same number of bits to be able to update the header afterwards when we know what data made it into the packet
-			Writer.WriteBits(0U, ReplicatedDestroyHeaderFlags_BitCount);
+			Writer.WriteBits((uint32)EReplicatedDestroyHeaderFlags::None, GetDestroyHeaderFlagsBitCount());
 		}
 
 		if (Writer.WriteBool(bHasState))
@@ -2912,21 +2906,27 @@ FReplicationWriter::EWriteObjectStatus FReplicationWriter::WriteObjectAndSubObje
 			// Rewrite destroy header if necessary	
 			if (bSentTearOff || Info.SubObjectPendingDestroy)
 			{
-				uint32 ReplicatedDestroyHeaderFlags = 0U;
+				EReplicatedDestroyHeaderFlags ReplicatedDestroyHeaderFlags = EReplicatedDestroyHeaderFlags::None;
 
-				// TearOff			
-				ReplicatedDestroyHeaderFlags |= bSentTearOff ? ReplicatedDestroyHeaderFlags_TearOff : ReplicatedDestroyHeaderFlags_None;
+				if (bSentTearOff)
+				{
+					ReplicatedDestroyHeaderFlags |= EReplicatedDestroyHeaderFlags::TearOff;
+				}
 
 				// Write SubObject destroy
 				if (Info.SubObjectPendingDestroy)
 				{
-					ReplicatedDestroyHeaderFlags |= ReplicatedDestroyHeaderFlags_EndReplication;
+					ReplicatedDestroyHeaderFlags |= EReplicatedDestroyHeaderFlags::IsSubObject;
+					
 					const bool bShouldDestroyInstance = ObjectData.RefHandle.IsDynamic() || NetRefHandleManager->IsDestroyedStartupObject(InternalIndex);					
-					ReplicatedDestroyHeaderFlags |= bShouldDestroyInstance ? ReplicatedDestroyHeaderFlags_DestroyInstance : ReplicatedDestroyHeaderFlags_None;
+					if (bShouldDestroyInstance)
+					{
+						ReplicatedDestroyHeaderFlags |= EReplicatedDestroyHeaderFlags::DestroySubObject;
+					}
 				}
 
 				FNetBitStreamWriteScope WriteScope(Writer, ReplicatedDestroyHeaderBitPos);
-				Writer.WriteBits(ReplicatedDestroyHeaderFlags, ReplicatedDestroyHeaderFlags_BitCount);
+				Writer.WriteBits((uint32)ReplicatedDestroyHeaderFlags, GetDestroyHeaderFlagsBitCount());
 			}
 			else if (bWriteBatchHeader && !(BatchEntry.bSentState || BatchEntry.bSentAttachments))
 			{
