Silence PVS compile warning

#rnx

#Synced-CL 48500075

diff --git a/Engine/Source/Runtime/Net/Iris/Private/Iris/ReplicationSystem/ReplicationReader.cpp b/Engine/Source/Runtime/Net/Iris/Private/Iris/ReplicationSystem/ReplicationReader.cpp
--- a/Engine/Source/Runtime/Net/Iris/Private/Iris/ReplicationSystem/ReplicationReader.cpp
+++ b/Engine/Source/Runtime/Net/Iris/Private/Iris/ReplicationSystem/ReplicationReader.cpp
@@ -596,7 +596,7 @@ void FReplicationReader::EndReplication(FInternalNetRefIndex InternalIndex, bool
 {
 	if (InternalIndex == ObjectIndexForOOBAttachment)
 	{
-		ensure(InternalIndex != ObjectIndexForOOBAttachment);
+		ensure(false);
 		return;
 	}
 	
