Iris
* More deprecated code removal
* Cleaned up EAddSubObjectFlags
* Removed IsOwnedSubObject concept. It was always set to true

#rnx
#Synced-CL 48743431

diff --git a/Engine/Source/Runtime/Net/Iris/Private/Iris/ReplicationSystem/Conditionals/ReplicationConditionals.cpp b/Engine/Source/Runtime/Net/Iris/Private/Iris/ReplicationSystem/Conditionals/ReplicationConditionals.cpp
--- a/Engine/Source/Runtime/Net/Iris/Private/Iris/ReplicationSystem/Conditionals/ReplicationConditionals.cpp
+++ b/Engine/Source/Runtime/Net/Iris/Private/Iris/ReplicationSystem/Conditionals/ReplicationConditionals.cpp
@@ -1074,14 +1074,11 @@ void FReplicationConditionals::InvalidateBaselinesForObjectHierarchy(uint32 Obje
 	for (const FInternalNetRefIndex SubObjectIndex : NetRefHandleManager->GetSubObjects(ObjectIndex))
 	{
 		const FNetRefHandleManager::FReplicatedObjectData& SubObjectData = NetRefHandleManager->GetReplicatedObjectDataNoCheck(SubObjectIndex);
-		if (SubObjectData.IsOwnedSubObject())
+		if (EnumHasAnyFlags(SubObjectData.Protocol->ProtocolTraits, EReplicationProtocolTraits::HasLifetimeConditionals))
 		{
-			if (EnumHasAnyFlags(SubObjectData.Protocol->ProtocolTraits, EReplicationProtocolTraits::HasLifetimeConditionals))
+			for (const uint32 ConnId : ConnectionsToInvalidate)
 			{
-				for (const uint32 ConnId : ConnectionsToInvalidate)
-				{
-					BaselineInvalidationTracker->InvalidateBaselines(SubObjectIndex, ConnId);
-				}
+				BaselineInvalidationTracker->InvalidateBaselines(SubObjectIndex, ConnId);
 			}
 		}
 	}
diff --git a/Engine/Source/Runtime/Net/Iris/Private/Iris/ReplicationSystem/Filtering/ReplicationFiltering.cpp b/Engine/Source/Runtime/Net/Iris/Private/Iris/ReplicationSystem/Filtering/ReplicationFiltering.cpp
--- a/Engine/Source/Runtime/Net/Iris/Private/Iris/ReplicationSystem/Filtering/ReplicationFiltering.cpp
+++ b/Engine/Source/Runtime/Net/Iris/Private/Iris/ReplicationSystem/Filtering/ReplicationFiltering.cpp
@@ -1054,11 +1054,8 @@ uint32 FReplicationFiltering::GetOwningConnectionIfDirty(uint32 ObjectIndex) con
 	if (NetRefHandleManager->GetSubObjectInternalIndices().GetBit(ObjectIndex))
 	{
 		const FNetRefHandleManager::FReplicatedObjectData& ObjectData = NetRefHandleManager->GetReplicatedObjectDataNoCheck(ObjectIndex);
-		if (ObjectData.IsOwnedSubObject())
-		{
-			const uint32 ParentIndex = NetRefHandleManager->GetReplicatedObjectDataNoCheck(ObjectIndex).SubObjectRootIndex;
-			return ObjectIndexToOwningConnection[ParentIndex];
-		}
+		const uint32 ParentIndex = NetRefHandleManager->GetReplicatedObjectDataNoCheck(ObjectIndex).SubObjectRootIndex;
+		return ObjectIndexToOwningConnection[ParentIndex];
 	}
 	
 	return ObjectIndexToOwningConnection[ObjectIndex];
@@ -1083,10 +1080,7 @@ void FReplicationFiltering::UpdateOwnerFiltering()
 			for (const FInternalNetRefIndex SubObjectIndex : NetRefHandleManager->GetSubObjects(ObjectIndex))
 			{
 				const FNetRefHandleManager::FReplicatedObjectData& ObjectData = NetRefHandleManager->GetReplicatedObjectDataNoCheck(SubObjectIndex);
-				if (ObjectData.IsOwnedSubObject())
-				{
-					ObjectIndexToOwningConnection[SubObjectIndex] = static_cast<uint16>(OwningConnectionId);
-				}
+				ObjectIndexToOwningConnection[SubObjectIndex] = static_cast<uint16>(OwningConnectionId);
 			}
 		};
 
diff --git a/Engine/Source/Runtime/Net/Iris/Private/Iris/ReplicationSystem/NetRefHandleManager.cpp b/Engine/Source/Runtime/Net/Iris/Private/Iris/ReplicationSystem/NetRefHandleManager.cpp
--- a/Engine/Source/Runtime/Net/Iris/Private/Iris/ReplicationSystem/NetRefHandleManager.cpp
+++ b/Engine/Source/Runtime/Net/Iris/Private/Iris/ReplicationSystem/NetRefHandleManager.cpp
@@ -671,7 +671,6 @@ bool FNetRefHandleManager::InternalAddSubObject(FInternalNetRefIndex RootObjectI
 	}
 	
 	SubObjectData.SubObjectRootIndex = RootObjectInternalIndex;
-	SubObjectData.bDestroySubObjectWithOwner = EnumHasAnyFlags(Flags, EAddSubObjectFlags::DestroyWithOwner);
 
 	// Mark the object as a subobject
 	SetIsSubObject(SubObjectInternalIndex, true);
@@ -757,7 +756,6 @@ void FNetRefHandleManager::InternalRemoveSubObject(FInternalNetRefIndex RootObje
 
 	SubObjectData.SubObjectRootIndex = InvalidInternalIndex;
 	SubObjectData.SubObjectParentIndex = InvalidInternalIndex;
-	SubObjectData.bDestroySubObjectWithOwner = false;
 
 	SetIsSubObject(SubObjectInternalIndex, false);
 }
diff --git a/Engine/Source/Runtime/Net/Iris/Private/Iris/ReplicationSystem/NetRefHandleManager.h b/Engine/Source/Runtime/Net/Iris/Private/Iris/ReplicationSystem/NetRefHandleManager.h
--- a/Engine/Source/Runtime/Net/Iris/Private/Iris/ReplicationSystem/NetRefHandleManager.h
+++ b/Engine/Source/Runtime/Net/Iris/Private/Iris/ReplicationSystem/NetRefHandleManager.h
@@ -47,14 +47,9 @@ typedef uint32 FInternalNetRefIndex;
 enum class EAddSubObjectFlags : uint32
 {
 	None = 0U,
-	WarnIfAlreadySubObject = 1U,
-	SkipIfAlreadySubObject = WarnIfAlreadySubObject << 1U,
-	DestroyWithOwner = SkipIfAlreadySubObject << 1U,
-	ReplicateWithSubObject = DestroyWithOwner << 1U,
+	ReplicateWithSubObject = 1U,
 	// Insert the subject at the start of the list
 	InsertAtStart = ReplicateWithSubObject << 1U, 
-
-	Default = WarnIfAlreadySubObject | DestroyWithOwner,
 };
 ENUM_CLASS_FLAGS(EAddSubObjectFlags);
 
@@ -116,7 +111,6 @@ public:
 
 		uint16 bShouldPropagateChangedStates : 1U = 0;
 		uint16 bTearOff : 1U = 0;
-		uint16 bDestroySubObjectWithOwner : 1U = 0;
 		uint16 bIsDependentObject : 1U = 0;
 		uint16 bHasDependentObjects : 1U = 0;
 		uint16 bAllowDestroyInstanceFromRemote : 1U = 0;
@@ -134,9 +128,6 @@ public:
 
 		ECloseReplicationReason CloseReason = ECloseReplicationReason::None;
 	
-		/** Returns true if this is a SubObject that gets destroyed when his RootObject is destroyed too. */
-		bool IsOwnedSubObject() const { return SubObjectRootIndex != InvalidInternalIndex && !!bDestroySubObjectWithOwner; }
-
 		/** Returns true if this is a SubObject */
 		bool IsSubObject() const { return SubObjectRootIndex != InvalidInternalIndex; }
 
@@ -286,8 +277,8 @@ public:
 	const FNetBitArray& GetSubObjectInternalIndices() const { return SubObjectInternalIndices; }
 	const FNetBitArrayView GetSubObjectInternalIndicesView() const { return MakeNetBitArrayView(SubObjectInternalIndices); }
 
-	bool AddSubObject(FNetRefHandle RootObjectHandle, FNetRefHandle SubObjectHandle, FNetRefHandle RelativeOtherSubObjectHandle, EAddSubObjectFlags Flags = EAddSubObjectFlags::Default);
-	bool AddSubObject(FNetRefHandle RootObjectHandle, FNetRefHandle SubObjectHandle, EAddSubObjectFlags Flags = EAddSubObjectFlags::Default);
+	bool AddSubObject(FNetRefHandle RootObjectHandle, FNetRefHandle SubObjectHandle, FNetRefHandle RelativeOtherSubObjectHandle, EAddSubObjectFlags Flags = EAddSubObjectFlags::None);
+	bool AddSubObject(FNetRefHandle RootObjectHandle, FNetRefHandle SubObjectHandle, EAddSubObjectFlags Flags = EAddSubObjectFlags::None);
 	void RemoveSubObject(FNetRefHandle SubObjectHandle);
 	
 	FNetRefHandle GetRootObjectOfAnyObject(FNetRefHandle NetRefHandle) const;
diff --git a/Engine/Source/Runtime/Net/Iris/Private/Iris/ReplicationSystem/ReplicationBridge.cpp b/Engine/Source/Runtime/Net/Iris/Private/Iris/ReplicationSystem/ReplicationBridge.cpp
--- a/Engine/Source/Runtime/Net/Iris/Private/Iris/ReplicationSystem/ReplicationBridge.cpp
+++ b/Engine/Source/Runtime/Net/Iris/Private/Iris/ReplicationSystem/ReplicationBridge.cpp
@@ -331,7 +331,7 @@ void UReplicationBridge::InternalAddSubObject(FNetRefHandle OwnerHandle, FNetRef
 	using namespace UE::Net;
 	using namespace UE::Net::Private;
 
-	EAddSubObjectFlags AddSubObjectFlags = EAddSubObjectFlags::Default;
+	EAddSubObjectFlags AddSubObjectFlags = EAddSubObjectFlags::None;
 
 	switch(InsertionOrder)
 	{
@@ -371,10 +371,9 @@ void UReplicationBridge::InternalDestroySubObjects(FNetRefHandle OwnerHandle, EE
 			FNetRefHandleManager::FReplicatedObjectData& SubObjectData = NetRefHandleManager->GetReplicatedObjectDataNoCheck(SubObjectInternalIndex);
 
 			const FNetRefHandle SubObjectHandle = SubObjectData.RefHandle;
-			const bool bDestroySubObjectWithOwner = SubObjectData.bDestroySubObjectWithOwner;
 				
 			// Tag subobject for destroy. The check against the scope is needed since the subobjects array might contain subobjects already pending destroy.
-			if (bDestroySubObjectWithOwner && NetRefHandleManager->IsScopableIndex(SubObjectInternalIndex))
+			if (NetRefHandleManager->IsScopableIndex(SubObjectInternalIndex))
 			{
 				SubObjectData.bPendingEndReplication = 1U;
 				UE_LOG_REPLICATIONBRIDGE(Verbose, TEXT("InternalDestroySubObjects %s"), *NetRefHandleManager->PrintObjectFromNetRefHandle(SubObjectHandle));
@@ -408,10 +407,9 @@ void UReplicationBridge::InternalRequestAsyncStopReplication(FNetRefHandle Handl
 			FNetRefHandleManager::FReplicatedObjectData& SubObjectData = NetRefHandleManager->GetReplicatedObjectDataNoCheck(SubObjectInternalIndex);
 
 			const FNetRefHandle SubObjectHandle = SubObjectData.RefHandle;
-			const bool bDestroySubObjectWithOwner = SubObjectData.bDestroySubObjectWithOwner;
-				
+
 			// Tag subobject for async stop replication. The check against the scope is needed since the subobjects array might contain destroyed subobjects
-			if (bDestroySubObjectWithOwner && NetRefHandleManager->IsScopableIndex(SubObjectInternalIndex))
+			if (NetRefHandleManager->IsScopableIndex(SubObjectInternalIndex))
 			{
 				UE_LOG_REPLICATIONBRIDGE(Verbose, TEXT("InternalRequestAsyncStopReplicationForSubObjects %s"), *NetRefHandleManager->PrintObjectFromNetRefHandle(SubObjectHandle));
 				InternalRequestAsyncStopReplication(SubObjectHandle, EndReplicationFlags);
diff --git a/Engine/Source/Runtime/Net/Iris/Private/Iris/ReplicationSystem/ReplicationWriter.cpp b/Engine/Source/Runtime/Net/Iris/Private/Iris/ReplicationSystem/ReplicationWriter.cpp
--- a/Engine/Source/Runtime/Net/Iris/Private/Iris/ReplicationSystem/ReplicationWriter.cpp
+++ b/Engine/Source/Runtime/Net/Iris/Private/Iris/ReplicationSystem/ReplicationWriter.cpp
@@ -743,29 +743,26 @@ void FReplicationWriter::SetPendingDestroyOrSubObjectPendingDestroyState(uint32
 	{
 		// Subobject destroyed before its owner is explicitly replicated as state data.
 		FNetRefHandleManager::FReplicatedObjectData& ObjectData = NetRefHandleManager->GetReplicatedObjectDataNoCheck(InternalIndex);
-		if (ObjectData.IsOwnedSubObject())
+		// If owner is not pending destroy we mark the state of the SubObject to SubObjectPendingDestroy and mark owner as having dirty subobjects which will 
+		// destroy the subobject using the replicated state path of the owner.
+		FReplicationInfo& OwnerInfo = ReplicatedObjects[ObjectData.SubObjectRootIndex];
+		if ((OwnerInfo.GetState() != EReplicatedObjectState::Invalid) && !ObjectsPendingDestroy.GetBit(ObjectData.SubObjectRootIndex))
 		{
-			// If owner is not pending destroy we mark the state of the SubObject to SubObjectPendingDestroy and mark owner as having dirty subobjects which will 
-			// destroy the subobject using the replicated state path of the owner.
-			FReplicationInfo& OwnerInfo = ReplicatedObjects[ObjectData.SubObjectRootIndex];
-			if ((OwnerInfo.GetState() != EReplicatedObjectState::Invalid) && !ObjectsPendingDestroy.GetBit(ObjectData.SubObjectRootIndex))
-			{
-				MarkObjectDirty(ObjectData.SubObjectRootIndex, "SetPendingDestroyOrSubObjectPendingDestroyState");
-				OwnerInfo.HasDirtySubObjects = 1U;
+			MarkObjectDirty(ObjectData.SubObjectRootIndex, "SetPendingDestroyOrSubObjectPendingDestroyState");
+			OwnerInfo.HasDirtySubObjects = 1U;
 
-				SetState(InternalIndex, EReplicatedObjectState::SubObjectPendingDestroy);
-				MarkObjectDirty(InternalIndex, "SetPendingDestroyOrSubObjectPendingDestroyState2");
-				ObjectsPendingDestroy.SetBit(InternalIndex);
-				Info.SubObjectPendingDestroy = 1U;
+			SetState(InternalIndex, EReplicatedObjectState::SubObjectPendingDestroy);
+			MarkObjectDirty(InternalIndex, "SetPendingDestroyOrSubObjectPendingDestroyState2");
+			ObjectsPendingDestroy.SetBit(InternalIndex);
+			Info.SubObjectPendingDestroy = 1U;
 			
-				// For clarity clear HasDirtyChangeMask as we do not intend to replicate state. The object itself needs to remain dirty in ObjectsWithDirtyChanges though.
-				Info.HasDirtyChangeMask = 0U;
-				// One cannot cancel a subobject destroy so we should not send any attachments for this subobject going forward
-				Info.HasAttachments = 0U;
+			// For clarity clear HasDirtyChangeMask as we do not intend to replicate state. The object itself needs to remain dirty in ObjectsWithDirtyChanges though.
+			Info.HasDirtyChangeMask = 0U;
+			// One cannot cancel a subobject destroy so we should not send any attachments for this subobject going forward
+			Info.HasAttachments = 0U;
 
-				ensure(!Info.TearOff || (GetFlushStatus(InternalIndex, Info, Info.FlushFlags) == FlushFlags_None));
-				return;
-			}
+			ensure(!Info.TearOff || (GetFlushStatus(InternalIndex, Info, Info.FlushFlags) == FlushFlags_None));
+			return;
 		}
 	}
 	else if (Info.HasDirtySubObjects)
@@ -997,19 +994,17 @@ void FReplicationWriter::UpdateDirtyGlobalLifetimeConditionals(TArrayView<FInter
 			if (Info.IsSubObject)
 			{
 				const FNetRefHandleManager::FReplicatedObjectData& ObjectData = NetRefHandleManager->GetReplicatedObjectDataNoCheck(InternalObjectIndex);
-				if (ObjectData.IsOwnedSubObject())
-				{
-					const FInternalNetRefIndex RootInternalObjectIndex = ObjectData.SubObjectRootIndex;
-					FReplicationInfo& OwnerInfo = ReplicatedObjects[RootInternalObjectIndex];
 
-					if (ObjectsInScope.GetBit(RootInternalObjectIndex) && (OwnerInfo.GetState() != EReplicatedObjectState::Invalid && OwnerInfo.GetState() < EReplicatedObjectState::PendingDestroy))
-					{
-						UE_LOG_REPWRITER(Verbose, TEXT("UpdateDirtyGlobalLifetimeConditionals for - %s"), ToCStr(NetRefHandleManager->PrintObjectFromIndex(RootInternalObjectIndex)));
+				const FInternalNetRefIndex RootInternalObjectIndex = ObjectData.SubObjectRootIndex;
+				FReplicationInfo& OwnerInfo = ReplicatedObjects[RootInternalObjectIndex];
+
+				if (ObjectsInScope.GetBit(RootInternalObjectIndex) && (OwnerInfo.GetState() != EReplicatedObjectState::Invalid && OwnerInfo.GetState() < EReplicatedObjectState::PendingDestroy))
+				{
+					UE_LOG_REPWRITER(Verbose, TEXT("UpdateDirtyGlobalLifetimeConditionals for - %s"), ToCStr(NetRefHandleManager->PrintObjectFromIndex(RootInternalObjectIndex)));
 						
-						MarkObjectDirty(RootInternalObjectIndex, "UpdateDirtyGlobalLifetimeConditionals");
-						OwnerInfo.HasDirtyConditionals = 1U;
-						OwnerInfo.HasDirtySubObjects = 1U;
-					}
+					MarkObjectDirty(RootInternalObjectIndex, "UpdateDirtyGlobalLifetimeConditionals");
+					OwnerInfo.HasDirtyConditionals = 1U;
+					OwnerInfo.HasDirtySubObjects = 1U;
 				}
 			}
 			else
@@ -1286,14 +1281,12 @@ void FReplicationWriter::HandleDeliveredRecord(const FReplicationRecord::FRecord
 					if (Info.IsSubObject)
 					{
 						FNetRefHandleManager::FReplicatedObjectData& ObjectData = NetRefHandleManager->GetReplicatedObjectDataNoCheck(InternalIndex);
-						if (ObjectData.IsOwnedSubObject())
+
+						FReplicationInfo& OwnerInfo = ReplicatedObjects[ObjectData.SubObjectRootIndex];
+						if (OwnerInfo.GetState() != EReplicatedObjectState::Invalid && OwnerInfo.GetState() < EReplicatedObjectState::PendingDestroy)
 						{
-							FReplicationInfo& OwnerInfo = ReplicatedObjects[ObjectData.SubObjectRootIndex];
-							if (OwnerInfo.GetState() != EReplicatedObjectState::Invalid && OwnerInfo.GetState() < EReplicatedObjectState::PendingDestroy)
-							{
-								MarkObjectDirty(ObjectData.SubObjectRootIndex, "HandleDeliveredRecordSubObjectTearOff");
-								OwnerInfo.HasDirtySubObjects = 1U;
-							}
+							MarkObjectDirty(ObjectData.SubObjectRootIndex, "HandleDeliveredRecordSubObjectTearOff");
+							OwnerInfo.HasDirtySubObjects = 1U;
 						}
 					}
 
@@ -1421,17 +1414,14 @@ void FReplicationWriter::HandleDeliveredRecord(const FReplicationRecord::FRecord
 			if (Info.IsSubObject)
 			{
 				FNetRefHandleManager::FReplicatedObjectData& ObjectData = NetRefHandleManager->GetReplicatedObjectDataNoCheck(InternalIndex);
-				if (ObjectData.IsOwnedSubObject())
+				FReplicationInfo& OwnerInfo = ReplicatedObjects[ObjectData.SubObjectRootIndex];
+				if (OwnerInfo.GetState() != EReplicatedObjectState::Invalid && OwnerInfo.GetState() < EReplicatedObjectState::PendingDestroy)
 				{
-					FReplicationInfo& OwnerInfo = ReplicatedObjects[ObjectData.SubObjectRootIndex];
-					if (OwnerInfo.GetState() != EReplicatedObjectState::Invalid && OwnerInfo.GetState() < EReplicatedObjectState::PendingDestroy)
-					{
-						MarkObjectDirty(ObjectData.SubObjectRootIndex, "HandleDeliveredRecordMarkHasAttachmentDirty");
-						OwnerInfo.HasDirtySubObjects = 1U;
+					MarkObjectDirty(ObjectData.SubObjectRootIndex, "HandleDeliveredRecordMarkHasAttachmentDirty");
+					OwnerInfo.HasDirtySubObjects = 1U;
 
-						// Bump prio if necessary
-						SchedulingPriorities[ObjectData.SubObjectRootIndex] = FPlatformMath::Max(FReplicationWriter::LostStatePriorityBump, SchedulingPriorities[ObjectData.SubObjectRootIndex]);
-					}
+					// Bump prio if necessary
+					SchedulingPriorities[ObjectData.SubObjectRootIndex] = FPlatformMath::Max(FReplicationWriter::LostStatePriorityBump, SchedulingPriorities[ObjectData.SubObjectRootIndex]);
 				}
 			}
 		}
@@ -1477,12 +1467,9 @@ void FReplicationWriter::HandleDeliveredRecord(const FReplicationRecord::FRecord
 				if (Info.IsSubObject)
 				{
 					FNetRefHandleManager::FReplicatedObjectData& ObjectData = NetRefHandleManager->GetReplicatedObjectDataNoCheck(InternalIndex);
-					if (ObjectData.IsOwnedSubObject())
-					{
-						FReplicationInfo& OwnerInfo = ReplicatedObjects[ObjectData.SubObjectRootIndex];
-						MarkObjectDirty(ObjectData.SubObjectRootIndex, "HandleDeliveredRecordTearOff");
-						OwnerInfo.HasDirtySubObjects = 1U;
-					}
+					FReplicationInfo& OwnerInfo = ReplicatedObjects[ObjectData.SubObjectRootIndex];
+					MarkObjectDirty(ObjectData.SubObjectRootIndex, "HandleDeliveredRecordTearOff");
+					OwnerInfo.HasDirtySubObjects = 1U;
 				}
 			}
 			else
@@ -1723,7 +1710,6 @@ void FReplicationWriter::HandleDroppedRecord<FReplicationWriter::EReplicatedObje
 			//ensureMsgf(false, TEXT("Subobjects destroyed individually should not be canceled. Object: %s"), ToCStr(NetRefHandleManager->PrintObjectFromIndex(InternalIndex)));
 			// If the subobject owner still is replicated and valid
 			FNetRefHandleManager::FReplicatedObjectData& ObjectData = NetRefHandleManager->GetReplicatedObjectDataNoCheck(InternalIndex);
-			check(ObjectData.IsOwnedSubObject());
 
 			FNetBitArrayView ChangeMask(Info.GetChangeMaskStoragePointer(), Info.ChangeMaskBitCount);
 			Info.HasDirtyChangeMask |= ChangeMask.IsAnyBitSet();
@@ -1807,7 +1793,6 @@ void FReplicationWriter::HandleDroppedRecord<FReplicationWriter::EReplicatedObje
 		{
 			// If the subobject owner still is replicated and valid
 			FNetRefHandleManager::FReplicatedObjectData& ObjectData = NetRefHandleManager->GetReplicatedObjectDataNoCheck(InternalIndex);
-			check(ObjectData.IsOwnedSubObject());
 
 			// If owner is not pending destroy we mark it as dirty so that we can replicate subobject destruction properly
 			// We might get away with not doing this if owner or subobject does not have any unconfirmed changes in flight.
@@ -2123,7 +2108,7 @@ uint32 FReplicationWriter::WriteRootObjectsPendingDestroy(FNetSerializationConte
 		FReplicationInfo& Info = GetReplicationInfo(InternalIndex);
 		const FNetRefHandleManager::FReplicatedObjectData& ObjectData = NetRefHandleManager->GetReplicatedObjectDataNoCheck(InternalIndex);
 
-		if (ObjectData.IsOwnedSubObject())
+		if (ObjectData.IsSubObject())
 		{
 			// Have the root object query status of subobjects
 			continue;
