Iris
- Added code that prevents creation of a NetObject for a NetRefHandle that still is replicating data to clients.
- Can be disabled with net.Iris.AllowCreateNetObjectForNetRefHandlePendingDestroy 1

#Synced-CL 48800604

diff --git a/Engine/Source/Runtime/Net/Iris/Private/Iris/ReplicationSystem/NetRefHandleManager.cpp b/Engine/Source/Runtime/Net/Iris/Private/Iris/ReplicationSystem/NetRefHandleManager.cpp
--- a/Engine/Source/Runtime/Net/Iris/Private/Iris/ReplicationSystem/NetRefHandleManager.cpp
+++ b/Engine/Source/Runtime/Net/Iris/Private/Iris/ReplicationSystem/NetRefHandleManager.cpp
@@ -26,11 +26,20 @@
 #include "UObject/Object.h"
 #include "UObject/UObjectGlobals.h"
 
+#include "HAL/IConsoleManager.h"
+
 CSV_DEFINE_CATEGORY(IrisCommon, true);
 
 namespace UE::Net::Private
 {
 
+static bool bAllowCreateNetObjectForNetRefHandlePendingDestroy = true;
+static FAutoConsoleVariableRef CVarAllowCreateNetObjectForNetRefHandlePendingDestroy(
+	TEXT("net.Iris.AllowCreateNetObjectForNetRefHandlePendingDestroy"),
+	bAllowCreateNetObjectForNetRefHandlePendingDestroy,
+	TEXT("If set to true we allow creation of a new NetObject for a NetRefHandle pending destroy that might still be replicating data to client. This is currently not fully supported and can lead to issues on the client.")
+);
+
 static bool IsValidIndex(FInternalNetRefIndex InternalIndex)
 {
 	return InternalIndex != FNetRefHandleManager::InvalidInternalIndex;
@@ -225,10 +234,32 @@ FInternalNetRefIndex FNetRefHandleManager::InternalCreateNetObject(const FNetRef
 	// Verify that the handle is free
 	if (RefHandleToInternalIndex.Contains(NetRefHandle))
 	{
+		UE_LOG(LogIris, Error, TEXT("NetRefHandleManager::InternalCreateNetObject %s already exists"), *NetRefHandle.ToString());
 		ensureMsgf(false, TEXT("NetRefHandleManager::InternalCreateNetObject %s already exists"), *NetRefHandle.ToString());
 		return InvalidInternalIndex;
 	}
 
+	// Verify that we are not trying to StartRepliation for an object using the same NetRefHandle as one that is pending destroy.
+	for (const FPendingDestroyInfo& PendingDestroyInfo : PendingDestroyInfos)
+	{
+		if (PendingDestroyInfo.RefHandle == NetRefHandle)
+		{
+			if (ReplicatedObjectRefCount[PendingDestroyInfo.InternalIndex] > 0)
+			{
+				if (!bAllowCreateNetObjectForNetRefHandlePendingDestroy)
+				{
+					UE_LOG(LogIris, Error, TEXT("NetRefHandleManager::InternalCreateNetObject Cannot create NetObject %s for object already replicated as %s marked for destroy. Only supported if using AsyncStopReplication."), *NetRefHandle.ToString(), *PrintObjectFromIndex(PendingDestroyInfo.InternalIndex));
+					ensureMsgf(false, TEXT("NetRefHandleManager::InternalCreateNetObject Cannot create NetObject %s for object marked for destroy. Only supported if using AsyncStopReplication."), *NetRefHandle.ToString());
+					return InvalidInternalIndex;
+				}
+				else
+				{
+					UE_LOG(LogIris, Warning, TEXT("NetRefHandleManager::InternalCreateNetObject Created NetObject %s for object already replicated as %s marked for destroy. This can lead to issues on clients. Only supported if using AsyncStopReplication."), *NetRefHandle.ToString(), *PrintObjectFromIndex(PendingDestroyInfo.InternalIndex));	
+				}
+			}
+		}
+	}
+
 	uint32 InternalIndex = GetNextFreeInternalIndex();
 
 	// Try to grow the NetObjectLists if no more indexes are available.
@@ -592,7 +623,7 @@ void FNetRefHandleManager::DestroyNetObject(FNetRefHandle RefHandle)
 		GlobalScopableInternalIndices.ClearBit(InternalIndex);
 
 		// We always defer the actual destroy
-		PendingDestroyInternalIndices.Add(InternalIndex);
+		PendingDestroyInfos.Emplace(InternalIndex, RefHandle);
 	}
 }
 
@@ -601,22 +632,22 @@ void FNetRefHandleManager::DestroyObjectsPendingDestroy()
 	IRIS_PROFILER_SCOPE(FNetRefHandleManager_DestroyObjectsPendingDestroy);
 
 	TArray<FInternalNetRefIndex> FreedInternalIndices;
-	FreedInternalIndices.Reserve(PendingDestroyInternalIndices.Num());
+	FreedInternalIndices.Reserve(PendingDestroyInfos.Num());
 
 	// Destroy Objects pending destroy
-	for (int32 It = PendingDestroyInternalIndices.Num() - 1; It >= 0; --It)
+	for (int32 It = PendingDestroyInfos.Num() - 1; It >= 0; --It)
 	{
-		const FInternalNetRefIndex InternalIndex = PendingDestroyInternalIndices[It];
+		const FInternalNetRefIndex InternalIndex = PendingDestroyInfos[It].InternalIndex;
 		// If we have subobjects pending tear off and such then wait before destroying the parent.
 		if (ReplicatedObjectRefCount[InternalIndex] == 0 && GetSubObjects(InternalIndex).Num() <= 0)
 		{
 			FreedInternalIndices.Add(InternalIndex);
 
 			InternalDestroyNetObject(InternalIndex);
-			PendingDestroyInternalIndices.RemoveAtSwap(It);
+			PendingDestroyInfos.RemoveAtSwap(It);
 		}
 	}
-	CSV_CUSTOM_STAT(IrisCommon, PendingDestroyInternalIndicesCount, (float)PendingDestroyInternalIndices.Num(), ECsvCustomStatOp::Set);
+	CSV_CUSTOM_STAT(IrisCommon, PendingDestroyInternalIndicesCount, (float)PendingDestroyInfos.Num(), ECsvCustomStatOp::Set);
 
 	if (OnInternalNetRefIndicesFreed.IsBound())
 	{
diff --git a/Engine/Source/Runtime/Net/Iris/Private/Iris/ReplicationSystem/NetRefHandleManager.h b/Engine/Source/Runtime/Net/Iris/Private/Iris/ReplicationSystem/NetRefHandleManager.h
--- a/Engine/Source/Runtime/Net/Iris/Private/Iris/ReplicationSystem/NetRefHandleManager.h
+++ b/Engine/Source/Runtime/Net/Iris/Private/Iris/ReplicationSystem/NetRefHandleManager.h
@@ -210,7 +210,13 @@ public:
 	// Mark object as no longer scopable, object will be removed from scope for all connections
 	void RemoveFromScope(FInternalNetRefIndex InternalIndex);
 
-	TArrayView<const FInternalNetRefIndex> GetObjectsPendingDestroy() const { return MakeArrayView(PendingDestroyInternalIndices); }
+	struct FPendingDestroyInfo
+	{
+		FInternalNetRefIndex InternalIndex;
+		FNetRefHandle RefHandle;
+	};
+
+	TConstArrayView<FPendingDestroyInfo> GetObjectsPendingDestroy() const { return MakeArrayView(PendingDestroyInfos); }
 	void DestroyObjectsPendingDestroy();
 
 	const FReplicatedObjectData& GetReplicatedObjectDataNoCheck(FInternalNetRefIndex InternalIndex) const { return ReplicatedObjectData[InternalIndex]; }
@@ -555,8 +561,8 @@ private:
 	// Map DestructionInfoIndex -> Orignal
 	TMap<uint32, uint32> DestroyedStartupObject;
 
-	// Array used in order to track objects pending destroy
-	TArray<FInternalNetRefIndex> PendingDestroyInternalIndices;
+	// Array used in order to track objects pending destroy/stopreplication
+	TArray<FPendingDestroyInfo> PendingDestroyInfos;
 
 	// Just an array containing data about our replicated objects
 	TNetChunkedArray<FReplicatedObjectData> ReplicatedObjectData;
diff --git a/Engine/Source/Runtime/Net/Iris/Private/Iris/ReplicationSystem/ReplicationReader.cpp b/Engine/Source/Runtime/Net/Iris/Private/Iris/ReplicationSystem/ReplicationReader.cpp
--- a/Engine/Source/Runtime/Net/Iris/Private/Iris/ReplicationSystem/ReplicationReader.cpp
+++ b/Engine/Source/Runtime/Net/Iris/Private/Iris/ReplicationSystem/ReplicationReader.cpp
@@ -2325,9 +2325,10 @@ void FReplicationReader::UpdateUnresolvableReferenceTracking()
 	TArray<uint32, TInlineAllocator<AssumedMaxDependentCount>> Dependents;
 
 	// Naively go through every object pending destroy, see if it's dynamic and update dependent's unresolved tracking
-	TArrayView<const FInternalNetRefIndex> ObjectsPendingDestroy = NetRefHandleManager->GetObjectsPendingDestroy();
-	for (const FInternalNetRefIndex InternalIndex : ObjectsPendingDestroy)
+	for (const FNetRefHandleManager::FPendingDestroyInfo& PendingDestroyInfo : NetRefHandleManager->GetObjectsPendingDestroy())
 	{
+		FInternalNetRefIndex InternalIndex = PendingDestroyInfo.InternalIndex;
+
 		const FNetRefHandleManager::FReplicatedObjectData& ObjectData = NetRefHandleManager->GetReplicatedObjectDataNoCheck(InternalIndex);
 		const FNetRefHandle DestroyedHandle = ObjectData.RefHandle;
 		if (!DestroyedHandle.IsDynamic())
