Removed NoSync on Name and Written Events to fix missing data when running Write Step in parallel.

#rnx
#tests Built and Tested against Dedicated Server Cooked Data
#Synced-CL 49043027

diff --git a/Engine/Source/Runtime/Net/Core/Private/Net/Core/Trace/Reporters/NetTraceReporter.cpp b/Engine/Source/Runtime/Net/Core/Private/Net/Core/Trace/Reporters/NetTraceReporter.cpp
--- a/Engine/Source/Runtime/Net/Core/Private/Net/Core/Trace/Reporters/NetTraceReporter.cpp
+++ b/Engine/Source/Runtime/Net/Core/Private/Net/Core/Trace/Reporters/NetTraceReporter.cpp
@@ -24,7 +24,7 @@ UE_TRACE_EVENT_BEGIN(NetTrace, InitEvent)
 UE_TRACE_EVENT_END()
 
 // Trace a name
-UE_TRACE_EVENT_BEGIN(NetTrace, NameEvent, NoSync)
+UE_TRACE_EVENT_BEGIN(NetTrace, NameEvent)
 	UE_TRACE_EVENT_FIELD(uint32, NameId)
 	UE_TRACE_EVENT_FIELD(UE::Trace::AnsiString, Name)
 UE_TRACE_EVENT_END()
@@ -142,13 +142,13 @@ UE_TRACE_EVENT_BEGIN(NetTrace, ObjectQuantizedEvent)
 	UE_TRACE_EVENT_FIELD(uint8, GameInstanceId)
 UE_TRACE_EVENT_END()
 
-UE_TRACE_EVENT_BEGIN(NetTrace, ObjectWrittenEvent, NoSync)
+UE_TRACE_EVENT_BEGIN(NetTrace, ObjectWrittenEvent)
 	UE_TRACE_EVENT_FIELD(uint64, ObjectId)
 	UE_TRACE_EVENT_FIELD(uint64, Cycles)
 	UE_TRACE_EVENT_FIELD(uint8, GameInstanceId)
 UE_TRACE_EVENT_END()
 
-UE_TRACE_EVENT_BEGIN(NetTrace, SubObjectEvent, NoSync)
+UE_TRACE_EVENT_BEGIN(NetTrace, SubObjectEvent)
 	UE_TRACE_EVENT_FIELD(uint64, ParentObjectId)
 	UE_TRACE_EVENT_FIELD(uint64, SubObjectId)
 	UE_TRACE_EVENT_FIELD(uint8, GameInstanceId)
