Iris Parallelization - NetConnection::Tick can be run in parallel on the server by setting net.iris.AllowParallelNetTick to true.
UChannels can indicate support for parallelism in SupportsParallelTick
Implemented thread safety for ObjectReferenceCache

#tests Built and Tested against Dedicated Server Cooked Data
[RN] Iris replication supports parallel NetConnection Tick to run the write phase across multiple threads
#Synced-CL 49043027

diff --git a/Engine/Source/Runtime/Engine/Classes/Engine/Channel.h b/Engine/Source/Runtime/Engine/Classes/Engine/Channel.h
--- a/Engine/Source/Runtime/Engine/Classes/Engine/Channel.h
+++ b/Engine/Source/Runtime/Engine/Classes/Engine/Channel.h
@@ -83,6 +83,7 @@ public:
 	uint32				bPooled:1;			// Set when placed in the actor channel pool
 	uint32				OpenedLocally:1;	// Whether channel was opened locally or by remote.
 	uint32				bOpenedForCheckpoint:1;	// Whether channel was opened by replay checkpoint recording
+	uint32				bSupportsParallelTick:1; // Indicates if this channel can be ticked in parallel with other NetConnection::Tick steps. Set this in the Init function of derived class
 	int32				ChIndex;			// Index of this channel.
 	FPacketIdRange		OpenPacketId;		// If OpenedLocally is true, this is the packet we sent the bOpen bunch on. Otherwise, it's the packet we received the bOpen bunch on.
 	FName				ChName;				// Name of the type of this channel.
@@ -185,6 +186,12 @@ public:
 	/** Returns true if this channel has flushed and acknowledged pending reliable data during a graceful close (when net.EnableGracefulClose is true) */
 	ENGINE_API virtual bool HasAcknowledgedAllReliableData() const;
 
+	/** Indicates if this channel can be ticked in parallel with other NetConnection::Tick steps*/
+	bool SupportsParallelTick() const
+	{
+		return bSupportsParallelTick;
+	}
+	
 protected:
 
 	/** Closes the actor channel but with a 'dormant' flag set so it can be reopened */
diff --git a/Engine/Source/Runtime/Engine/Classes/Engine/Engine.h b/Engine/Source/Runtime/Engine/Classes/Engine/Engine.h
--- a/Engine/Source/Runtime/Engine/Classes/Engine/Engine.h
+++ b/Engine/Source/Runtime/Engine/Classes/Engine/Engine.h
@@ -230,11 +230,16 @@ struct FNetDriverDefinition
 	UPROPERTY()
 	int32 MaxChannelsOverride;
 
+	/** Runs the Server Net Connection Tick in parallel */
+	UPROPERTY()
+	bool bRunParallelConnectionTick;
+
 	FNetDriverDefinition() :
 		DefName(NAME_None),
 		DriverClassName(NAME_None),
 		DriverClassNameFallback(NAME_None),
-		MaxChannelsOverride(INDEX_NONE)
+		MaxChannelsOverride(INDEX_NONE),
+		bRunParallelConnectionTick(false)
 	{
 	}
 };
@@ -281,6 +286,12 @@ struct FIrisNetDriverConfig
 	 */
 	UPROPERTY()
 	EReplicationSystem PreferredReplicationSystem = EReplicationSystem::Default;
+
+	/**
+	 * If true, and Iris is used as the replication system then NetConnection::Tick can be run in parallel on the server for this NetDriver
+	 */
+	UPROPERTY()
+	bool bCanUseParallelNetConnectionTick = false;
 };
 
 
@@ -3278,6 +3289,9 @@ public:
 
 	/** Returns true if the netdriver will run with Iris enable. */
 	ENGINE_API bool WillNetDriverUseIris(const FWorldContext& Context, FName InNetDriverDefinition, FName InNetDriverName) const;
+
+	/** Returns true if the netdriver will run server NetConnection::Tick in parallel (only when using Iris) */
+	ENGINE_API bool WillNetDriverUseParallelNetConnectionTick(const FWorldContext& Context, const FName InNetDriverDefinition, const FName InNetDriverName) const;
 	
 	/** A configurable list of actors that are automatically spawned upon server startup (just prior to InitGame) */
 	UPROPERTY(config)
diff --git a/Engine/Source/Runtime/Engine/Classes/Engine/NetConnection.h b/Engine/Source/Runtime/Engine/Classes/Engine/NetConnection.h
--- a/Engine/Source/Runtime/Engine/Classes/Engine/NetConnection.h
+++ b/Engine/Source/Runtime/Engine/Classes/Engine/NetConnection.h
@@ -13,6 +13,7 @@
 #include "Serialization/CustomVersion.h"
 #include "Misc/EngineVersion.h"
 #include "Misc/NetworkGuid.h"
+#include "Misc/TransactionallySafeCriticalSection.h"
 #include "GameFramework/OnlineReplStructs.h"
 #include "GameFramework/UpdateLevelVisibilityLevelInfo.h"
 #include "Engine/NetDriver.h"
@@ -660,6 +661,11 @@ public:
 private:
 	FCustomVersionContainer NetworkCustomVersions;
 
+public:
+#if UE_SUPPORT_PARALLEL_IRIS
+	FTransactionallySafeCriticalSection& GetSharedNetTickMutex();
+#endif // UE_SUPPORT_PARALLEL_IRIS
+
 public:
 
 	// Log tracking
diff --git a/Engine/Source/Runtime/Engine/Classes/Engine/NetDriver.h b/Engine/Source/Runtime/Engine/Classes/Engine/NetDriver.h
--- a/Engine/Source/Runtime/Engine/Classes/Engine/NetDriver.h
+++ b/Engine/Source/Runtime/Engine/Classes/Engine/NetDriver.h
@@ -24,6 +24,7 @@
 #include "Net/Core/Connection/ConnectionHandle.h"
 #include "Net/Core/Misc/DDoSDetection.h"
 #include "IPAddress.h"
+#include "Misc/TransactionallySafeCriticalSection.h"
 #include "Net/NetAnalyticsTypes.h"
 #include "Net/NetConnectionIdHandler.h"
 #include "HAL/LowLevelMemTracker.h"
@@ -351,6 +352,7 @@ struct FReplicatedStaticActorDestructionInfo;
 enum class ECreateReplicationChangelistMgrFlags;
 enum class EEngineNetworkRuntimeFeatures : uint16;
 class UReplicationSystem;
+class UReplicationBridge;
 namespace UE::Net
 {
 	class FNetObjectGroupHandle;
@@ -948,6 +950,9 @@ public:
 	UPROPERTY(Config)
 	float TimeoutMultiplierForUnoptimizedBuilds;
 
+	/** Cached at the beginning of net driver TickFlush step so that parallel threads can access it safely */
+	int32 CachedEngineMaxTickRate = 0;
+
 	/** Connection to the server (this net driver is a client) */
 	UPROPERTY()
 	TObjectPtr<class UNetConnection> ServerConnection;
@@ -1118,7 +1123,7 @@ public:
 	FName GetNetDriverDefinition() const { return NetDriverDefinition; }
 
 	/** Callback after the engine created the NetDriver and set our name for the first time */
-	ENGINE_API void PostCreation(bool bInitializeWithIris);
+	ENGINE_API void PostCreation(bool bInitializeWithIris, bool bEnableParallelNetTick);
 
 	/** Reset some NetDriver settings after destroying and recreating a replication system or replication driver */
 	ENGINE_API void ReinitBase();
@@ -2534,6 +2539,9 @@ private:
 	/** True when the NetDriver has been configured to run with the Iris replication system.*/
 	bool bIsUsingIris = false;
 
+	/** True when the NetDriver has been configured to run NetConnection::Tick in parallel*/
+	bool bIsUsingParallelNetTick = false;
+
 	// For FindOrAddNetworkObjectInfo
 	TPimplPtr<FNetworkObjectInfo> DummyNetworkObjectInfo;
 
diff --git a/Engine/Source/Runtime/Engine/Private/Net/Experimental/Iris/DataStreamChannel.cpp b/Engine/Source/Runtime/Engine/Private/Net/Experimental/Iris/DataStreamChannel.cpp
--- a/Engine/Source/Runtime/Engine/Private/Net/Experimental/Iris/DataStreamChannel.cpp
+++ b/Engine/Source/Runtime/Engine/Private/Net/Experimental/Iris/DataStreamChannel.cpp
@@ -59,6 +59,7 @@ UDataStreamChannel::UDataStreamChannel(const FObjectInitializer& ObjectInitializ
 , bHandshakeComplete(0U)
 {
 	ChName = FName("DataStream");
+	bSupportsParallelTick = true;
 }
 
 void UDataStreamChannel::InitDataStreamManager()
@@ -322,10 +323,18 @@ bool UDataStreamChannel::WriteData(UE::Net::EDataStreamWriteMode WriteMode)
 		BeginWriteParams.MaxPackets = 1U;
 	}
 
+	check(Connection);
+	UNetDriver* ConnDriver = Connection->GetDriver();
+	UReplicationSystem* RepSystem = ConnDriver ? ConnDriver->GetReplicationSystem() : nullptr;
+	ensure(RepSystem);
+
+	//Get a NetStatsContext that is unique per-thread so locking NetStats is not required for the rest of this function
+	UE::Net::FScopedNetStatsContext ScopedNetStatsContext(RepSystem);
+
 	// Currently we want to use a full bunch so we flush if we have to
 	bool bNeedsPreSendFlush = Connection->SendBuffer.GetNumBits() > MAX_PACKET_HEADER_BITS;
 
-	auto WriteDataFunction = [this, &bNeedsPreSendFlush, &BeginWriteParams, MaxBunchBits, MaxBytes]() -> UDataStream::EWriteResult
+	auto WriteDataFunction = [this, &bNeedsPreSendFlush, &BeginWriteParams, MaxBunchBits, MaxBytes, &ScopedNetStatsContext]() -> UDataStream::EWriteResult
 	{
 		if (bNeedsPreSendFlush)
 		{
@@ -345,6 +354,8 @@ bool UDataStreamChannel::WriteData(UE::Net::EDataStreamWriteMode WriteMode)
 
 		FNetSerializationContext SerializationContext(&BitWriter);
 		SerializationContext.SetPacketId(Connection->OutPacketId);
+		SerializationContext.SetNetStatsContext(ScopedNetStatsContext.GetRawPtr());
+		
 
 #if UE_NET_TRACE_ENABLED	
 		// For Iris we can use the connection trace collector as long as we make sure that the packet is prepared 
diff --git a/Engine/Source/Runtime/Engine/Private/NetConnection.cpp b/Engine/Source/Runtime/Engine/Private/NetConnection.cpp
--- a/Engine/Source/Runtime/Engine/Private/NetConnection.cpp
+++ b/Engine/Source/Runtime/Engine/Private/NetConnection.cpp
@@ -319,6 +319,10 @@ static SIZE_T CountBytes(FWrittenChannelsRecord& WrittenChannelsRecord);
 namespace UE::Net::Private
 {
 	extern bool bTrackDormantObjectsByLevel;
+
+#if UE_SUPPORT_PARALLEL_IRIS
+	static FTransactionallySafeCriticalSection SharedNetTickMutex; //Shared across all NetConnections to protect sections of the Tick logic that cannot run in parallel
+#endif // UE_SUPPORT_PARALLEL_IRIS
 }
 
 /*-----------------------------------------------------------------------------
@@ -1098,6 +1102,8 @@ void UNetConnection::Close(FNetResult&& CloseReason)
 	{
 		NETWORK_PROFILER(GNetworkProfiler.TrackEvent(TEXT("CLOSE"), *(GetName() + TEXT(" ") + LowLevelGetRemoteAddress()), this));
 
+		UE_IRIS_PARALLEL_EXPR(UE::TScopeLock ScopeLock(GetSharedNetTickMutex()));
+
 		TStringBuilder<2048> CloseStr;
 
 		CloseStr.Append(TEXT("UNetConnection::Close: "));
@@ -2512,8 +2518,12 @@ void UNetConnection::FlushNet(bool bIgnoreSimulation)
 		++OutPackets;
 		++OutPacketsThisFrame;
 		++OutTotalPackets;
-		Driver->OutPackets++;
-		Driver->OutTotalPackets++;
+		{
+			UE_IRIS_PARALLEL_EXPR(Driver->LockNetDriverStats(true));
+			Driver->OutPackets++;
+			Driver->OutTotalPackets++;
+			UE_IRIS_PARALLEL_EXPR(Driver->UnlockNetDriverStats(true));
+		}
 
 		//Record the first packet time in the histogram
 		if (!bFlushedNetThisFrame)
@@ -2884,7 +2894,7 @@ void UNetConnection::WriteDummyPacketInfo(FBitWriter& Writer)
 			   Writer.GetNumBits() - BitsWrittenPrePacketInfo, (1 + NetConnectionHelper::NumBitsForJitterClockTimeInHeader));
 #endif
 
-		const uint8 bHasServerFrameTime = Driver->IsServer() ? bLastHasServerFrameTime : (CVarPingExcludeFrameTime.GetValueOnGameThread() > 0 ? 1u : 0u);
+		const uint8 bHasServerFrameTime = Driver->IsServer() ? bLastHasServerFrameTime : (CVarPingExcludeFrameTime.GetValueOnAnyThread() > 0 ? 1u : 0u);
 		Writer.WriteBit(bHasServerFrameTime);
 
 		if (bHasServerFrameTime && Driver->IsServer())
@@ -4718,6 +4728,10 @@ void UNetConnection::Tick(float DeltaSeconds)
 
 	AssertValid();
 
+	// Acquire a coarse grain lock which protects access to shared state such as NetDriver internals, lag simulation, disconnect logic etc
+	// It is unlocked whilst ticking individual channels which is where the most time consuming work is done
+	UE_IRIS_PARALLEL_EXPR(UE::TScopeLock ScopeLock(GetSharedNetTickMutex()));
+
 	// Get frame time.
 	const double CurrentRealtimeSeconds = FPlatformTime::Seconds();
 
@@ -4736,7 +4750,7 @@ void UNetConnection::Tick(float DeltaSeconds)
 
 	FrameTime = CurrentRealtimeSeconds - LastTime;
 	const int32 MaxNetTickRate = Driver->MaxNetTickRate;
-	float EngineTickRate = GEngine->GetMaxTickRate(0.0f, false);
+	float EngineTickRate = Driver->CachedEngineMaxTickRate;
 	// We want to make sure the DesiredTickRate stays at <= 0 if there's no tick rate limiting of any kind, since it's used later in the function for bandwidth limiting.
 	if (MaxNetTickRate > 0 && EngineTickRate <= 0.0f)
 	{
@@ -4906,8 +4920,9 @@ void UNetConnection::Tick(float DeltaSeconds)
 		// Tick the channels.
 		if (CVarTickAllOpenChannels.GetValueOnAnyThread() == 0)
 		{
-			for( int32 i=ChannelsToTick.Num()-1; i>=0; i-- )
+			for (int32 i = ChannelsToTick.Num()-1; i >= 0; i--)
 			{
+				UE_IRIS_PARALLEL_EXPR(UE::TScopeUnlock<FTransactionallySafeCriticalSection> ScopeUnlock(ChannelsToTick[i]->SupportsParallelTick() ? &GetSharedNetTickMutex() : nullptr ));
 				ChannelsToTick[i]->Tick();
 
 				if (ChannelsToTick[i]->CanStopTicking())
@@ -4920,6 +4935,7 @@ void UNetConnection::Tick(float DeltaSeconds)
 		{
 			for (int32 i = OpenChannels.Num() - 1; i >= 0; i--)
 			{
+				UE_IRIS_PARALLEL_EXPR(UE::TScopeUnlock<FTransactionallySafeCriticalSection> ScopeUnlock(OpenChannels[i]->SupportsParallelTick() ? &GetSharedNetTickMutex() : nullptr ));
 				if (OpenChannels[i])
 				{
 					OpenChannels[i]->Tick();
@@ -4990,6 +5006,7 @@ void UNetConnection::Tick(float DeltaSeconds)
 		// Delay any packet sends on the server, until we've verified that a packet has been received from the client.
 		if (bHandlerHandshakeComplete && HasReceivedClientPacket())
 		{
+			UE_IRIS_PARALLEL_EXPR(UE::TScopeUnlock<FTransactionallySafeCriticalSection> ScopeUnlock(&GetSharedNetTickMutex()));
 			FlushNet();
 		}
 	}
@@ -6207,6 +6224,13 @@ void UNetConnection::SetNetworkCustomVersions(const FCustomVersionContainer& Cus
 	}
 }
 
+#if UE_SUPPORT_PARALLEL_IRIS
+FTransactionallySafeCriticalSection& UNetConnection::GetSharedNetTickMutex()
+{
+	return UE::Net::Private::SharedNetTickMutex;
+}
+#endif // UE_SUPPORT_PARALLEL_IRIS
+
 /*-----------------------------------------------------------------------------
 	USimulatedClientNetConnection.
 -----------------------------------------------------------------------------*/
diff --git a/Engine/Source/Runtime/Engine/Private/NetDriver.cpp b/Engine/Source/Runtime/Engine/Private/NetDriver.cpp
--- a/Engine/Source/Runtime/Engine/Private/NetDriver.cpp
+++ b/Engine/Source/Runtime/Engine/Private/NetDriver.cpp
@@ -186,6 +186,40 @@ extern int32 GNetRPCDebug;
 
 namespace UE::Net::Private
 {
+	/** 
+	* ConnectionTickTask is used to run each individual connection's Tick step in parallel
+	* When running these tasks, everything in NetConnection->Tick is expected to be thread safe
+	*/
+	struct FConnectionTickTaskData
+	{
+		float DeltaSeconds = 0.0f;
+
+		FConnectionTickTaskData(float InDeltaSeconds)
+			: DeltaSeconds(InDeltaSeconds)
+		{}
+	};
+
+	class FConnectionTickTask
+	{
+	public:
+		FConnectionTickTask(UNetConnection* InConnection, FConnectionTickTaskData& InConnectionTickTaskData)
+			: NetConnection(InConnection)
+			, ConnectionTickTaskData(InConnectionTickTaskData)
+		{
+		}
+
+		FORCENOINLINE void DoTask() const
+		{
+			CSV_SCOPED_TIMING_STAT_EXCLUSIVE(ReplicationNetConnectionTickTask);
+
+			NetConnection->Tick(ConnectionTickTaskData.DeltaSeconds);
+		}
+
+	protected:
+		UNetConnection* NetConnection = nullptr;
+		FConnectionTickTaskData& ConnectionTickTaskData;
+	};
+
 	struct FAutoDestructProperty
 	{
 		FProperty* Property = nullptr;
@@ -340,6 +374,12 @@ namespace UE::Net::Private
 		TEXT("net.iris.UpdateReplicationSystemWithNoConnections"),
 		bUpdateReplicationSystemWithNoConnections,
 		TEXT("If true ReplicationSystem::NetUpdate will be called even if we have no connections to ensure that we do not accumulate pending work."));
+
+	bool bAllowParallelNetTickAtRuntime = false;
+	static FAutoConsoleVariableRef CVarNetAllowParallelNetTick(
+		TEXT("net.iris.AllowParallelNetTick"),
+		bAllowParallelNetTickAtRuntime,
+		TEXT("If true, NetConnection::Tick can be run in parallel. This is designed as a runtime disable flag to be used alongside config changes. To enable parallel tick, it also requires the NetDriver config bCanUseParallelNetConnectionTick=true and ReplicationSystemConfig bAllowParallelTasks=true."));
 	
 #if UE_WITH_REMOTE_OBJECT_HANDLE
 	bool bServerReplicateActorsGranularTransactions = true;
@@ -1140,6 +1180,9 @@ void UNetDriver::TickFlush(float DeltaSeconds)
 	}
 	FSimpleScopeSecondsCounter ScopedTimer(GTickFlushGameDriverTimeSeconds, bEnableTimer);
 
+	// Cached here to avoid calling ShouldThrottleCPUUsage during parallel steps
+	CachedEngineMaxTickRate = GEngine->GetMaxTickRate(0.0f, false);
+
 	if (IsServer() && (ClientConnections.Num() > 0 || UE::Net::Private::bUpdateReplicationSystemWithNoConnections) && !bSkipServerReplicateActors)
 	{
 		// Update all clients.
@@ -1193,9 +1236,43 @@ void UNetDriver::TickFlush(float DeltaSeconds)
 		// Queue up any voice packets the server has locally
 		ProcessLocalServerPackets();
 	}
+
+	// Server UNetConnection::Tick can be run in parallel if running Iris and this NetDriver supports it
+	if (IsServer() && bIsUsingParallelNetTick && ReplicationSystem && ReplicationSystemConfigServer.bAllowParallelTasks && UE::Net::Private::bAllowParallelNetTickAtRuntime)
 	{
 		QUICK_SCOPE_CYCLE_COUNTER(STAT_NetDriver_TickClientConnections)
 
+		// Set up a Task for each ClientConnection to run its Tick in parallel
+		TArray<UE::Tasks::TTask<void>> TasksToComplete;
+		TArray<UE::Net::Private::FConnectionTickTaskData> TaskData;
+		TaskData.Init(UE::Net::Private::FConnectionTickTaskData(DeltaSeconds), ClientConnections.Num());
+		int32 CurrentConnectionIndex = 0;
+
+		ReplicationSystem->StartParallelPhase();
+
+		for (UNetConnection* Connection : ClientConnections)
+		{
+			IRIS_PROFILER_SCOPE(UConnectionTick_PrepareAndDispatchTask);
+			
+			UE::Net::Private::FConnectionTickTaskData& Data = TaskData[CurrentConnectionIndex];
+
+			TasksToComplete.Add(UE::Tasks::Launch(UE_SOURCE_LOCATION, [this, Connection, &Data]()
+				{
+					UE::Net::Private::FConnectionTickTask NewTask = UE::Net::Private::FConnectionTickTask(Connection, Data);
+					NewTask.DoTask();
+				}
+			));
+
+			CurrentConnectionIndex++;
+		}
+	
+		// Block until all tick tasks are done
+		UE::Tasks::Wait(TasksToComplete);
+
+		ReplicationSystem->StopParallelPhase();
+	}
+	else
+	{
 		for (UNetConnection* Connection : ClientConnections)
 		{
 			Connection->Tick(DeltaSeconds);
@@ -2617,7 +2694,8 @@ void UNetDriver::TryUpgradeNetworkFeatures(EEngineNetworkRuntimeFeatures RemoteF
 		SetReplicationDriver(nullptr);
 
 		constexpr bool bInitWithIris = true;
-		PostCreation(bInitWithIris);
+		constexpr bool bUsingParallelNetTick = false;
+		PostCreation(bInitWithIris, bUsingParallelNetTick);
 
 		RecreateIrisSystem();
 		
@@ -2646,7 +2724,8 @@ void UNetDriver::TryUpgradeNetworkFeatures(EEngineNetworkRuntimeFeatures RemoteF
 		DestroyIrisSystem();
 
 		constexpr bool bInitWithoutIris = false;
-		PostCreation(bInitWithoutIris);
+		constexpr bool bUsingParallelNetTick = false;
+		PostCreation(bInitWithoutIris, bUsingParallelNetTick);
 
 		ReinitBase();
 
@@ -6480,9 +6559,10 @@ void UNetDriver::ReinitBase()
 	UE_LOG(LogNet, Log, TEXT("Re-InitBase %s (NetDriverDefinition %s) using replication model %s"), *NetDriverName.ToString(), *NetDriverDefinition.ToString(), *GetReplicationModelName());
 }
 
-void UNetDriver::PostCreation(bool bInitializeWithIris)
+void UNetDriver::PostCreation(bool bInitializeWithIris, bool bEnableParallelNetTick)
 {
 	bIsUsingIris = bInitializeWithIris;
+	bIsUsingParallelNetTick = bEnableParallelNetTick;
 
 	if (NetDriverDefinition == NAME_GameNetDriver)
 	{
diff --git a/Engine/Source/Runtime/Engine/Private/UnrealEngine.cpp b/Engine/Source/Runtime/Engine/Private/UnrealEngine.cpp
--- a/Engine/Source/Runtime/Engine/Private/UnrealEngine.cpp
+++ b/Engine/Source/Runtime/Engine/Private/UnrealEngine.cpp
@@ -14532,6 +14532,13 @@ bool UEngine::WillNetDriverUseIris(const FWorldContext& Context, FName InNetDriv
 	return bUseIrisRepSystem;
 }
 
+bool UEngine::WillNetDriverUseParallelNetConnectionTick(const FWorldContext& Context, const FName InNetDriverDefinition, const FName InNetDriverName) const
+{
+	const FIrisNetDriverConfig* IrisConfig = GetIrisNetDriverConfig(InNetDriverDefinition, InNetDriverName);
+	const bool bWithParallelNetTick = IrisConfig ? IrisConfig->bCanUseParallelNetConnectionTick : false;
+	return bWithParallelNetTick;
+}
+
 namespace UE::Private
 {
 	UNetDriver* CreateNetDriver_Local(UEngine* Engine, FWorldContext& Context, FName NetDriverDefinition, FName InNetDriverName)
@@ -14670,6 +14677,7 @@ namespace UE::Private
 
 				const FName DriverName = InNetDriverName.IsNone() ? ReturnVal->GetFName() : InNetDriverName;
 				const bool bInitializeWithIris = Engine->WillNetDriverUseIris(Context, NetDriverDefinition, DriverName);
+				const bool bUseParallelNetConnectionTick = Engine->WillNetDriverUseParallelNetConnectionTick(Context, NetDriverDefinition, DriverName);
 
 				if (bInitializeWithIris)
 				{
@@ -14678,7 +14686,7 @@ namespace UE::Private
 
 				ReturnVal->SetNetDriverName(DriverName);
 				ReturnVal->SetNetDriverDefinition(NetDriverDefinition);
-				ReturnVal->PostCreation(bInitializeWithIris);
+				ReturnVal->PostCreation(bInitializeWithIris, bUseParallelNetConnectionTick);
 
 				new(Context.ActiveNetDrivers) FNamedNetDriver(ReturnVal, Definition);
 
diff --git a/Engine/Source/Runtime/Net/Iris/Private/Iris/ReplicationSystem/Conditionals/ReplicationConditionals.cpp b/Engine/Source/Runtime/Net/Iris/Private/Iris/ReplicationSystem/Conditionals/ReplicationConditionals.cpp
--- a/Engine/Source/Runtime/Net/Iris/Private/Iris/ReplicationSystem/Conditionals/ReplicationConditionals.cpp
+++ b/Engine/Source/Runtime/Net/Iris/Private/Iris/ReplicationSystem/Conditionals/ReplicationConditionals.cpp
@@ -575,7 +575,7 @@ void FReplicationConditionals::Update()
 	}
 }
 
-void FReplicationConditionals::GetChildSubObjectsToReplicate(uint32 ReplicatingConnectionId, const FConditionalsMask& LifetimeConditionals,  const FInternalNetRefIndex ParentObjectIndex, FSubObjectsToReplicateArray& OutSubObjectsToReplicate)
+void FReplicationConditionals::GetChildSubObjectsToReplicate(uint32 ReplicatingConnectionId, const FConditionalsMask& LifetimeConditionals,  const FInternalNetRefIndex ParentObjectIndex, FSubObjectsToReplicateArray& OutSubObjectsToReplicate) const
 {
 	// To mimic old system we use a weird replication order based on hierarchy, SubSubObjects are replicated before the parent
 	FChildSubObjectsInfo SubObjectsInfo;
diff --git a/Engine/Source/Runtime/Net/Iris/Private/Iris/ReplicationSystem/Conditionals/ReplicationConditionals.h b/Engine/Source/Runtime/Net/Iris/Private/Iris/ReplicationSystem/Conditionals/ReplicationConditionals.h
--- a/Engine/Source/Runtime/Net/Iris/Private/Iris/ReplicationSystem/Conditionals/ReplicationConditionals.h
+++ b/Engine/Source/Runtime/Net/Iris/Private/Iris/ReplicationSystem/Conditionals/ReplicationConditionals.h
@@ -141,7 +141,7 @@ private:
 	void ClearPerObjectInfo(FInternalNetRefIndex ObjectIndex);
 	void ClearConnectionInfosForObject(const FNetBitArray& ValidConnections, FInternalNetRefIndex ObjectIndex);
 
-	void GetChildSubObjectsToReplicate(uint32 ReplicatingConnectionId, const FConditionalsMask& LifetimeConditionals, const FInternalNetRefIndex ParentObjectIndex, FSubObjectsToReplicateArray& OutSubObjectsToReplicate);
+	void GetChildSubObjectsToReplicate(uint32 ReplicatingConnectionId, const FConditionalsMask& LifetimeConditionals, const FInternalNetRefIndex ParentObjectIndex, FSubObjectsToReplicateArray& OutSubObjectsToReplicate) const;
 
 	ELifetimeCondition GetDynamicCondition(FInternalNetRefIndex ObjectIndex, uint16 RepIndex) const;
 	void SetDynamicCondition(FInternalNetRefIndex ObjectIndex, uint16 RepIndex, ELifetimeCondition Condition);
diff --git a/Engine/Source/Runtime/Net/Iris/Private/Iris/ReplicationSystem/NetRefHandleManager.cpp b/Engine/Source/Runtime/Net/Iris/Private/Iris/ReplicationSystem/NetRefHandleManager.cpp
--- a/Engine/Source/Runtime/Net/Iris/Private/Iris/ReplicationSystem/NetRefHandleManager.cpp
+++ b/Engine/Source/Runtime/Net/Iris/Private/Iris/ReplicationSystem/NetRefHandleManager.cpp
@@ -10,6 +10,8 @@
 #include "Iris/Serialization/InternalNetSerializationContext.h"
 #include "Iris/Serialization/NetSerializationContext.h"
 
+#include "Misc/ScopeLock.h"
+
 #include "Net/Core/Trace/NetDebugName.h"
 #include "Net/Core/Trace/NetTrace.h"
 
@@ -369,6 +371,7 @@ bool FNetRefHandleManager::HasInstanceProtocol(FInternalNetRefIndex InternalInde
 
 FNetRefHandle FNetRefHandleManager::AllocateNetRefHandle(bool bIsStatic)
 {
+	UE_IRIS_PARALLEL_EXPR(UE::TScopeLock Lock(HandleSerialCriticalSection));
 	uint64& NextHandleSerial = bIsStatic ? NextStaticHandleIndex : NextDynamicHandleIndex;
 
 	FNetRefHandle NewHandle = AllocateNetRefHandleFromSerial(NextHandleSerial, bIsStatic);
@@ -381,6 +384,7 @@ FNetRefHandle FNetRefHandleManager::AllocateNetRefHandle(bool bIsStatic)
 
 FNetRefHandle FNetRefHandleManager::AllocateNetRefHandleFromSerial(uint64 Serial, bool bIsStatic)
 {
+	UE_IRIS_PARALLEL_EXPR(UE::TScopeLock Lock(HandleSerialCriticalSection));
 	FNetRefHandle NewHandle = MakeNetRefHandleFromSerial(Serial, bIsStatic, ReplicationSystemId);
 
 	// Verify that the handle is free
diff --git a/Engine/Source/Runtime/Net/Iris/Private/Iris/ReplicationSystem/NetRefHandleManager.h b/Engine/Source/Runtime/Net/Iris/Private/Iris/ReplicationSystem/NetRefHandleManager.h
--- a/Engine/Source/Runtime/Net/Iris/Private/Iris/ReplicationSystem/NetRefHandleManager.h
+++ b/Engine/Source/Runtime/Net/Iris/Private/Iris/ReplicationSystem/NetRefHandleManager.h
@@ -20,6 +20,8 @@
 #include "Iris/ReplicationSystem/NetRefHandle.h"
 #include "Iris/ReplicationSystem/ObjectReferenceTypes.h"
 
+#include "Misc/TransactionallySafeCriticalSection.h"
+
 #include "Net/Core/NetBitArray.h"
 #include "Net/Core/NetHandle/NetHandle.h"
 
@@ -593,6 +595,10 @@ private:
 
 	// Track NetBitArrays owned by the NetRefHandleManager so they can automatically be resized.
 	TArray<FNetBitArray*> OwnedNetBitArrays;
+#if UE_SUPPORT_PARALLEL_IRIS
+	// Protects parallel access to NextHandleId and RefHandleToInternalIndex when allocating NetRefHandles
+	FTransactionallySafeCriticalSection HandleSerialCriticalSection;
+#endif // UE_SUPPORT_PARALLEL_IRIS
 };
 
 const FNetRefHandleManager::FReplicatedObjectData& FNetRefHandleManager::GetReplicatedObjectData(FInternalNetRefIndex InternalIndex) const
diff --git a/Engine/Source/Runtime/Net/Iris/Private/Iris/ReplicationSystem/ObjectReferenceCache.cpp b/Engine/Source/Runtime/Net/Iris/Private/Iris/ReplicationSystem/ObjectReferenceCache.cpp
--- a/Engine/Source/Runtime/Net/Iris/Private/Iris/ReplicationSystem/ObjectReferenceCache.cpp
+++ b/Engine/Source/Runtime/Net/Iris/Private/Iris/ReplicationSystem/ObjectReferenceCache.cpp
@@ -200,6 +200,8 @@ bool FObjectReferenceCache::ShouldIgnoreWhenMissing(FNetRefHandle RefHandle) con
 		return false;		// Server never ignores when missing, always warns
 	}
 
+	UE_IRIS_PARALLEL_EXPR(UE::TScopeLock InternalLock(ObjectReferenceCacheInternalCS));
+
 	const FCachedNetObjectReference* CacheObject = ReferenceHandleToCachedReference.Find(RefHandle);
 	if (CacheObject == nullptr)
 	{
@@ -274,7 +276,9 @@ FNetRefHandle FObjectReferenceCache::PreRegisterObjectReferenceHandle(const UObj
 	if (IsAuthority())
 	{
 		FNetObjectReference Ref;
+		UE_IRIS_PARALLEL_EXPR(UE::TScopeLock InternalLock(ObjectReferenceCacheInternalCS));
 		CreateObjectReferenceInternal(Object, Ref, ReplicationSystem);
+
 		FCachedNetObjectReference* CachedRef = ReferenceHandleToCachedReference.Find(Ref.GetRefHandle());
 		if (CachedRef)
 		{
@@ -303,6 +307,8 @@ bool FObjectReferenceCache::CreateObjectReferenceInternal(const UObject* Object,
 
 	CA_ASSUME(Object != nullptr);
 
+	UE_IRIS_PARALLEL_EXPR(UE::TScopeLock InternalLock(ObjectReferenceCacheInternalCS));
+
 	// Check if we are already know about this object
 	if (FNetRefHandle* RefHandlePtr = ObjectToNetReferenceHandle.Find(Object))
 	{
@@ -537,7 +543,7 @@ FNetRefHandle FObjectReferenceCache::GetObjectReferenceHandleFromObject(const UO
 	if (const FNetRefHandle* Reference = ObjectToNetReferenceHandle.Find(Object))
 	{
 		// Verify that this is the same object
-
+		UE_IRIS_PARALLEL_EXPR(UE::TScopeLock InternalLock(ObjectReferenceCacheInternalCS));
 		const FCachedNetObjectReference* CachedObjectPtr = ReferenceHandleToCachedReference.Find(*Reference);
 		if (CachedObjectPtr == nullptr)
 		{
@@ -574,7 +580,6 @@ void FObjectReferenceCache::BindRemoteReference(FNetRefHandle RefHandle, const U
 	check(RefHandle.IsCompleteHandle());
 
 	UE_LOG_REFERENCECACHE(Verbose, TEXT("%hs: Binding Handle: %s to Object: %s (0x%p)" ), __FUNCTION__, *RefHandle.ToString(), *GetNameSafe(Object), Object);
-
 #if UE_NET_VALIDATE_REFERENCECACHE
 	// Dynamic object might have been referenced before it gets replicated, then it will then already be registered (with a path relative to its outer)
 	if (FNetRefHandle* ExistingHandle = ObjectToNetReferenceHandle.Find(Object))
@@ -628,7 +633,6 @@ void FObjectReferenceCache::AddPreRegisteredReference(FNetRefHandle RefHandle, c
 		FNetRefHandle CompleteRefHandle = FNetRefHandleManager::MakeNetRefHandle(RefHandle.GetId(), ReplicationSystem->GetId());
 
 		BindRemoteReference(CompleteRefHandle, Object);
-
 		if (FCachedNetObjectReference* CacheObjectPtr = ReferenceHandleToCachedReference.Find(CompleteRefHandle))
 		{
 			CacheObjectPtr->bIsPreRegistered = true;
@@ -756,6 +760,7 @@ void FObjectReferenceCache::RemoveReference(FNetRefHandle RefHandle, const UObje
 
 UObject* FObjectReferenceCache::GetObjectFromReferenceHandle(FNetRefHandle RefHandle)
 {
+	UE_IRIS_PARALLEL_EXPR(UE::TScopeLock InternalLock(ObjectReferenceCacheInternalCS));
 	FCachedNetObjectReference* CachedObject = ReferenceHandleToCachedReference.Find(RefHandle);
 
 	UObject* ResolvedObject = CachedObject ? CachedObject->Object.Get() : nullptr;
@@ -765,6 +770,7 @@ UObject* FObjectReferenceCache::GetObjectFromReferenceHandle(FNetRefHandle RefHa
 
 UObject* FObjectReferenceCache::GetPreRegisteredObjectFromReferenceHandle(FNetRefHandle RefHandle)
 {
+	UE_IRIS_PARALLEL_EXPR(UE::TScopeLock InternalLock(ObjectReferenceCacheInternalCS));
 	FCachedNetObjectReference* CachedObject = ReferenceHandleToCachedReference.Find(RefHandle);
 
 	UObject* ResolvedObject = (CachedObject && CachedObject->bIsPreRegistered) ? CachedObject->Object.Get() : nullptr;
@@ -774,12 +780,14 @@ UObject* FObjectReferenceCache::GetPreRegisteredObjectFromReferenceHandle(FNetRe
 
 bool FObjectReferenceCache::IsNetRefHandleBroken(FNetRefHandle RefHandle, bool bMustBeRegistered) const
 {
+	UE_IRIS_PARALLEL_EXPR(UE::TScopeLock InternalLock(ObjectReferenceCacheInternalCS));
 	const FCachedNetObjectReference* CacheObjectPtr = ReferenceHandleToCachedReference.Find(RefHandle);	
 	return CacheObjectPtr ? CacheObjectPtr->bIsBroken : bMustBeRegistered;
 }
 
 bool FObjectReferenceCache::IsNetRefHandlePreRegistered(FNetRefHandle RefHandle) const
 {
+	UE_IRIS_PARALLEL_EXPR(UE::TScopeLock InternalLock(ObjectReferenceCacheInternalCS));
 	const FCachedNetObjectReference* CacheObjectPtr = ReferenceHandleToCachedReference.Find(RefHandle);
 	return CacheObjectPtr ? CacheObjectPtr->bIsPreRegistered : false;
 }
@@ -795,6 +803,7 @@ bool FObjectReferenceCache::IsNetRefHandlePending(FNetRefHandle NetRefHandle, co
 			return true;
 		}
 
+		UE_IRIS_PARALLEL_EXPR(UE::TScopeLock InternalLock(ObjectReferenceCacheInternalCS));
 		const FCachedNetObjectReference* CacheObjectPtr = ReferenceHandleToCachedReference.Find(NetRefHandle);
 
 		if (!CacheObjectPtr)
@@ -815,6 +824,7 @@ bool FObjectReferenceCache::IsNetRefHandlePending(FNetRefHandle NetRefHandle, co
 
 UObject* FObjectReferenceCache::ResolveObjectReferenceHandleInternal(FNetRefHandle RefHandle, const FNetObjectResolveContext& ResolveContext, bool& bOutMustBeMapped)
 {
+	UE_IRIS_PARALLEL_EXPR(UE::TScopeLock InternalLock(ObjectReferenceCacheInternalCS));
 	FCachedNetObjectReference* CacheObjectPtr = ReferenceHandleToCachedReference.Find(RefHandle);
 	
 	if (!CacheObjectPtr)
@@ -1278,6 +1288,7 @@ FNetObjectReference FObjectReferenceCache::GetReplicatedOuter(const FNetObjectRe
 		return FNetObjectReference();
 	}
 
+	UE_IRIS_PARALLEL_EXPR(UE::TScopeLock InternalLock(ObjectReferenceCacheInternalCS));
 	const FCachedNetObjectReference* CachedObject = ReferenceHandleToCachedReference.Find(Reference.GetRefHandle());
 	FNetRefHandle OuterNetRefHandle = CachedObject ? CachedObject->OuterNetRefHandle : FNetRefHandle::GetInvalid();
 
@@ -1298,6 +1309,7 @@ FNetObjectReference FObjectReferenceCache::GetReplicatedOuter(const FNetObjectRe
 
 FNetRefHandle FObjectReferenceCache::GetDynamicRoot(const FNetRefHandle Handle) const
 {
+	UE_IRIS_PARALLEL_EXPR(UE::TScopeLock InternalLock(ObjectReferenceCacheInternalCS));
 	const FCachedNetObjectReference* CachedObject = ReferenceHandleToCachedReference.Find(Handle);
 	while (CachedObject && CachedObject->OuterNetRefHandle.IsValid())
 	{
@@ -1356,6 +1368,7 @@ FNetObjectReference FObjectReferenceCache::GetOrCreateObjectReference(const FStr
 		CachedObject.bIsPending = 0U;
 		CachedObject.bIsBroken = 0U;
 
+		UE_IRIS_PARALLEL_EXPR(UE::TScopeLock InternalLock(ObjectReferenceCacheInternalCS));
 		ReferenceHandleToCachedReference.Add(RefHandle, CachedObject);
 	
 		UE_LOG_REFERENCECACHE(Verbose, TEXT("ObjectReferenceCache::GetOrCreateObjectReference(FromPath) Added cache for %s for Object: %s, OuterHandle %s OuterObject: %s (0x%p)"), 
@@ -1373,6 +1386,8 @@ void FObjectReferenceCache::WriteFullReferenceInternal(FNetSerializationContext&
 	const FNetRefHandle RefHandle = Ref.GetRefHandle();
 
 	UE_NET_TRACE_OBJECT_SCOPE(RefHandle, *Writer, Context.GetTraceCollector(), ENetTraceVerbosity::Verbose);
+
+	UE_IRIS_PARALLEL_EXPR(UE::TScopeLock InternalLock(ObjectReferenceCacheInternalCS));
 	
 	// This is unfortunate, but as we defer writing of references, we can end up with stale references
 	const FCachedNetObjectReference* CachedObject = RefHandle.IsValid() ? ReferenceHandleToCachedReference.Find(RefHandle) : nullptr;
@@ -1670,6 +1685,8 @@ void FObjectReferenceCache::WriteReference(FNetSerializationContext& Context, FN
 		return;
 	}
 
+	UE_IRIS_PARALLEL_EXPR(UE::TScopeLock InternalLock(ObjectReferenceCacheInternalCS));
+
 	// This is unfortunate, but as we defer writing of references, we can end up with stale references so we need to do this lookup
 	const FCachedNetObjectReference* CachedObject = RefHandle.IsValid() ? ReferenceHandleToCachedReference.Find(RefHandle) : nullptr;
 
@@ -1743,6 +1760,8 @@ void FObjectReferenceCache::ReadReference(FNetSerializationContext& Context, FNe
 
 	UE_NET_TRACE_SET_SCOPE_OBJECTID(ReferenceScope, NetRefHandle);
 
+	UE_IRIS_PARALLEL_EXPR(UE::TScopeLock InternalLock(ObjectReferenceCacheInternalCS));
+
 	// The handle must exist in the cache or it must be a dynamic handle that will be explicitly registered
 	if (ensureMsgf(NetRefHandle.IsDynamic() || ReferenceHandleToCachedReference.Find(NetRefHandle), TEXT("FObjectReferenceCache::ReadReference: Handle %s must already be known or be dynamic"), *NetRefHandle.ToString()))
 	{
@@ -1964,6 +1983,8 @@ bool FObjectReferenceCache::WriteMustBeMappedExports(FNetSerializationContext& C
 		// Only write the loading priority before the first reference
 		bool bHasWrittenPriority = false;
 
+		UE_IRIS_PARALLEL_EXPR(UE::TScopeLock InternalLock(ObjectReferenceCacheInternalCS));
+
 		for (const FNetObjectReference& Reference : ExportsView)
 		{
 			const FNetRefHandle Handle = Reference.GetRefHandle();
@@ -2088,6 +2109,8 @@ void FObjectReferenceCache::GenerateFullPath_r(FNetRefHandle RefHandle, const FN
 		return;
 	}
 
+	UE_IRIS_PARALLEL_EXPR(UE::TScopeLock InternalLock(ObjectReferenceCacheInternalCS));
+
 	const FCachedNetObjectReference* CacheObject = ReferenceHandleToCachedReference.Find(RefHandle);
 
 	if (CacheObject == nullptr)
@@ -2235,6 +2258,8 @@ void FObjectReferenceCache::AsyncPackageCallback(const FName& PackageName, UPack
 	{
 		const bool bIsBroken = (Package == nullptr);
 
+		UE_IRIS_PARALLEL_EXPR(UE::TScopeLock InternalLock(ObjectReferenceCacheInternalCS));
+
 		for (FNetRefHandle RefHandleToProcess : PendingLoadRequest->NetRefHandles)
 		{
 			if (FCachedNetObjectReference* CacheObject = ReferenceHandleToCachedReference.Find(RefHandleToProcess))
@@ -2309,6 +2334,8 @@ void FObjectReferenceCache::AsyncPackageForcedFailCallback(const FName& PackageN
 
 	if (FPendingAsyncLoadRequest const* const PendingLoadRequest = PendingAsyncLoadRequests.Find(PackageName))
 	{
+		UE_IRIS_PARALLEL_EXPR(UE::TScopeLock InternalLock(ObjectReferenceCacheInternalCS));
+
 		for (FNetRefHandle RefHandleToProcess : PendingLoadRequest->NetRefHandles)
 		{
 			if (FCachedNetObjectReference* CacheObject = ReferenceHandleToCachedReference.Find(RefHandleToProcess))
@@ -2410,6 +2437,8 @@ void FObjectReferenceCache::RemoveTrackedQueuedBatchObjectReference(const FNetRe
 
 const TCHAR* FObjectReferenceCache::GetObjectRelativePath(FNetRefHandle NetRefHandle, UReplicationSystem* ReplicationSystem) const
 {
+	UE_IRIS_PARALLEL_EXPR(UE::TScopeLock InternalLock(ObjectReferenceCacheInternalCS));
+
 	if (const FCachedNetObjectReference* CacheObject = ReferenceHandleToCachedReference.Find(NetRefHandle))
 	{
 		const TCHAR* ResolvedToken = CacheObject->RelativePath.ResolveToken();
diff --git a/Engine/Source/Runtime/Net/Iris/Private/Iris/ReplicationSystem/ObjectReferenceCache.h b/Engine/Source/Runtime/Net/Iris/Private/Iris/ReplicationSystem/ObjectReferenceCache.h
--- a/Engine/Source/Runtime/Net/Iris/Private/Iris/ReplicationSystem/ObjectReferenceCache.h
+++ b/Engine/Source/Runtime/Net/Iris/Private/Iris/ReplicationSystem/ObjectReferenceCache.h
@@ -262,6 +262,12 @@ private:
 
 private:
 
+#if UE_SUPPORT_PARALLEL_IRIS
+	// Protects access to internal state such as the reference maps. Does not protect access to TokenStores or NetRefHandleManager
+	// Lock must be held whilst accessing pointers to FCachedNetObjectReference in ReferenceHandleToCachedReference
+	mutable FTransactionallySafeCriticalSection ObjectReferenceCacheInternalCS;
+#endif // UE_SUPPORT_PARALLEL_IRIS
+
 	// Map raw UObject pointer -> Handle
 	// To verify that the reference is valid we need to check the weakpointer stored in the cache
 	TMap<const UObject*, FNetRefHandle> ObjectToNetReferenceHandle;
diff --git a/Engine/Source/Runtime/Net/Iris/Private/Iris/ReplicationSystem/ObjectReplicationBridge.cpp b/Engine/Source/Runtime/Net/Iris/Private/Iris/ReplicationSystem/ObjectReplicationBridge.cpp
--- a/Engine/Source/Runtime/Net/Iris/Private/Iris/ReplicationSystem/ObjectReplicationBridge.cpp
+++ b/Engine/Source/Runtime/Net/Iris/Private/Iris/ReplicationSystem/ObjectReplicationBridge.cpp
@@ -1117,6 +1117,7 @@ void UObjectReplicationBridge::RemoveCreationDependencyLink(FNetRefHandle Parent
 
 bool UObjectReplicationBridge::WriteNetRefHandleCreationInfo(FReplicationBridgeSerializationContext& Context, FNetRefHandle Handle)
 {
+	UE_IRIS_PARALLEL_EXPR(FScopeLock ScopeLock(&WriteNetRefHandleCreationCS));
 	using namespace UE::Net;
 	using namespace UE::Net::Private;
 
diff --git a/Engine/Source/Runtime/Net/Iris/Private/Iris/ReplicationSystem/ReplicationSystem.cpp b/Engine/Source/Runtime/Net/Iris/Private/Iris/ReplicationSystem/ReplicationSystem.cpp
--- a/Engine/Source/Runtime/Net/Iris/Private/Iris/ReplicationSystem/ReplicationSystem.cpp
+++ b/Engine/Source/Runtime/Net/Iris/Private/Iris/ReplicationSystem/ReplicationSystem.cpp
@@ -1,5 +1,5 @@
 // Copyright Epic Games, Inc. All Rights Reserved.
-
+ 
 #include "Iris/ReplicationSystem/ReplicationSystem.h"
 #include "Iris/Core/IrisCsv.h"
 #include "Iris/Core/IrisDebugging.h"
@@ -35,6 +35,9 @@
 #include "Iris/Serialization/NetSerializer.h"
 #include "Iris/Serialization/IrisObjectReferencePackageMap.h"
 
+#include "Iris/Stats/NetStats.h"
+#include "Iris/Stats/NetStatsContext.h"
+
 #include "Iris/Metrics/NetMetrics.h"
 
 #include "HAL/IConsoleManager.h"
@@ -920,6 +923,25 @@ bool UReplicationSystem::IsInParallelPhase() const
 	return GetReplicationSystemInternal()->IsInParallelPhase();
 }
 
+UE::Net::Private::FNetStatsContext* UReplicationSystem::AcquireStatsForSerializationContext()
+{
+	const bool bIsParallel = IsInParallelPhase();
+	check(GetReplicationSystemInternal());
+	UE::Net::Private::FNetTypeStats& NetTypeStats = GetReplicationSystemInternal()->GetNetTypeStats();
+	UE::Net::Private::FNetStatsContext* NetStatsContext = bIsParallel ? NetTypeStats.AcquireChildNetStatsContext() : NetTypeStats.GetNetStatsContext();
+	return NetStatsContext;
+}
+
+void UReplicationSystem::ReleaseStatsForSerializationContext(UE::Net::Private::FNetStatsContext* NetStatsContext)
+{
+	const bool bIsParallel = IsInParallelPhase();
+	check(GetReplicationSystemInternal());
+	UE::Net::Private::FNetTypeStats& NetTypeStats = GetReplicationSystemInternal()->GetNetTypeStats();
+	if (bIsParallel && ensure(NetStatsContext))
+	{
+		NetTypeStats.ReleaseChildNetStatsContext(NetStatsContext);
+	}
+}
 
 void UReplicationSystem::NetUpdate(float DeltaSeconds)
 {
@@ -2329,6 +2351,32 @@ TArrayView<UReplicationSystem*> FReplicationSystemFactory::GetAllReplicationSyst
 	return TArrayView<UReplicationSystem*>(ReplicationSystems);
 }
 
+FScopedNetStatsContext::FScopedNetStatsContext(UReplicationSystem* const InReplicationSystem)
+{
+	if (ensure(InReplicationSystem))
+	{
+		StoredNetStatsContext = InReplicationSystem->AcquireStatsForSerializationContext();
+		StoredReplicationSystem = InReplicationSystem;
+		bWasAcquiredInParallelPhase = InReplicationSystem->IsInParallelPhase();
+	}
+}
+
+FScopedNetStatsContext::~FScopedNetStatsContext()
+{
+	if (ensure(StoredReplicationSystem))
+	{
+		ensureMsgf(bWasAcquiredInParallelPhase == StoredReplicationSystem->IsInParallelPhase(), TEXT("ScopedNetStatsContext lifetime mismatch - %s"), bWasAcquiredInParallelPhase ? TEXT("Expected to be destroyed in ParallelPhase") : TEXT("Expected to be destroyed outside of ParallelPhase"));
+	
+		StoredReplicationSystem->ReleaseStatsForSerializationContext(StoredNetStatsContext);
+		StoredNetStatsContext = nullptr;
+	}
+}
+
+Private::FNetStatsContext* FScopedNetStatsContext::GetRawPtr()
+{
+	return StoredNetStatsContext;
+}
+
 } //end namespace UE::Net
 
 #pragma endregion
diff --git a/Engine/Source/Runtime/Net/Iris/Private/Iris/ReplicationSystem/ReplicationSystemInternal.h b/Engine/Source/Runtime/Net/Iris/Private/Iris/ReplicationSystem/ReplicationSystemInternal.h
--- a/Engine/Source/Runtime/Net/Iris/Private/Iris/ReplicationSystem/ReplicationSystemInternal.h
+++ b/Engine/Source/Runtime/Net/Iris/Private/Iris/ReplicationSystem/ReplicationSystemInternal.h
@@ -183,6 +183,27 @@ public:
 	{
 		GetReplicationBridge()->PostSeamlessTravelGarbageCollect();
 	}
+	#if UE_SUPPORT_PARALLEL_IRIS
+	void AcquireTickReplicationStatsWriteLock()
+	{
+		TickReplicationStatsMutex.WriteLock();
+	}
+		
+	void AcquireTickReplicationStatsReadLock()
+	{
+		TickReplicationStatsMutex.ReadLock();
+	}
+
+	void ReleaseTickReplicationStatsWriteLock()
+	{
+		TickReplicationStatsMutex.WriteUnlock();
+	}
+
+	void ReleaseTickReplicationStatsReadLock()
+	{
+		TickReplicationStatsMutex.ReadUnlock();
+	}
+#endif // UE_SUPPORT_PARALLEL_IRIS
 
 private:
 	FReplicationProtocolManager ReplicationProtocolManager;
@@ -211,6 +232,11 @@ private:
 	FForwardNetRPCCallMulticastDelegate ForwardNetRPCCallMulticastDelegate;
 	uint32 Id;
 
+#if UE_SUPPORT_PARALLEL_IRIS
+	/** Mutex which protects access to TickReplicationStats during parallel phases*/
+	FTransactionallySafeRWLock TickReplicationStatsMutex;
+#endif // UE_SUPPORT_PARALLEL_IRIS
+
 	/** When true this prevents any changes to the filter system. Enabled during times where adding filter options is unsupported. */
 	bool bBlockFilterChanges = false;
 
diff --git a/Engine/Source/Runtime/Net/Iris/Private/Iris/ReplicationSystem/ReplicationWriter.cpp b/Engine/Source/Runtime/Net/Iris/Private/Iris/ReplicationSystem/ReplicationWriter.cpp
--- a/Engine/Source/Runtime/Net/Iris/Private/Iris/ReplicationSystem/ReplicationWriter.cpp
+++ b/Engine/Source/Runtime/Net/Iris/Private/Iris/ReplicationSystem/ReplicationWriter.cpp
@@ -4164,6 +4164,8 @@ void FReplicationWriter::EndWrite()
 		const uint32 NumPendingAdditionallySchedule1dObjects = WriteContext.AdditionalObjectSchedulingStack.Num();
 		const uint32 NumPendingObjectsToWrite = WriteContext.ScheduledObjectCount - WriteContext.CurrentIndex;
 
+		//Update TickReplicationStats (lock must be acquired to be threadsafe)
+		UE_IRIS_PARALLEL_EXPR(Parameters.ReplicationSystem->GetReplicationSystemInternal()->AcquireTickReplicationStatsWriteLock());
 		FReplicationStats& ReplicationStats = Parameters.ReplicationSystem->GetReplicationSystemInternal()->GetTickReplicationStats();
 		ReplicationStats.PendingObjectCount += NumPendingObjectsToWrite;
 		ReplicationStats.PendingDependentObjectCount += NumPendingAdditionallySchedule1dObjects;
@@ -4172,6 +4174,7 @@ void FReplicationWriter::EndWrite()
 		ReplicationStats.MaxPendingDependentObjectCount = FMath::Max(ReplicationStats.MaxPendingDependentObjectCount, NumPendingAdditionallySchedule1dObjects);
 		ReplicationStats.MaxHugeObjectSendQueue = FMath::Max(ReplicationStats.MaxHugeObjectSendQueue, HugeObjectSendQueue.NumRootObjectsInTransit());
 		++ReplicationStats.SampleCount;
+		UE_IRIS_PARALLEL_EXPR(Parameters.ReplicationSystem->GetReplicationSystemInternal()->ReleaseTickReplicationStatsWriteLock());
 
  #if UE_NET_ENABLE_REPLICATIONWRITER_LOG
 		// See if we failed to write any objects
@@ -4242,7 +4245,6 @@ UDataStream::EWriteResult FReplicationWriter::Write(FNetSerializationContext& Co
 	FInternalNetSerializationContext InternalContext(Parameters.ReplicationSystem);
 	Context.SetLocalConnectionId(Parameters.ConnectionId);
 	Context.SetInternalContext(&InternalContext);
-	Context.SetNetStatsContext(NetTypeStats->GetNetStatsContext());
 
 	// Give some info for the case when we consider splitting a huge object.
 	WriteBitStreamInfo.ReplicationStartPos = Writer.GetPosBits();
diff --git a/Engine/Source/Runtime/Net/Iris/Private/Iris/Stats/NetStats.h b/Engine/Source/Runtime/Net/Iris/Private/Iris/Stats/NetStats.h
--- a/Engine/Source/Runtime/Net/Iris/Private/Iris/Stats/NetStats.h
+++ b/Engine/Source/Runtime/Net/Iris/Private/Iris/Stats/NetStats.h
@@ -196,7 +196,7 @@ public:
 	int32 GetOrCreateTypeStats(FName Name);
 
 	/** Get parent context if stats is enabled. Should not be called when in a parallel phase, as the ChildContexts should be used instead. */
-	FNetStatsContext* GetNetStatsContext() { check(!bIsInParallelPhase); return IsEnabled() ?  ParentStatsContext : nullptr; }
+	FNetStatsContext* GetNetStatsContext() { ensure(!bIsInParallelPhase); return IsEnabled() ?  ParentStatsContext : nullptr; }
 
 	/** Updated every frame based on the state of the CSVProfiler */
 	bool IsEnabled() const { return bIsEnabled; }
diff --git a/Engine/Source/Runtime/Net/Iris/Public/Iris/ReplicationSystem/ObjectReplicationBridge.h b/Engine/Source/Runtime/Net/Iris/Public/Iris/ReplicationSystem/ObjectReplicationBridge.h
--- a/Engine/Source/Runtime/Net/Iris/Public/Iris/ReplicationSystem/ObjectReplicationBridge.h
+++ b/Engine/Source/Runtime/Net/Iris/Public/Iris/ReplicationSystem/ObjectReplicationBridge.h
@@ -685,6 +685,11 @@ private:
 
 	/** List of class names that will not prevent replication to happen when a protocol mismatch is detected.*/
 	TArray<FString> ClassNamesToIgnoreProtocolMismatch;
+#if UE_SUPPORT_PARALLEL_IRIS
+private:
+	/** Protects access to CachedCreationHeaders inside WriteNetRefHandleCreationInfo */
+	FCriticalSection WriteNetRefHandleCreationCS;
+#endif // UE_SUPPORT_PARALLEL_IRIS
 };
 
 //------------------------------------------------------------------------
diff --git a/Engine/Source/Runtime/Net/Iris/Public/Iris/ReplicationSystem/ReplicationSystem.h b/Engine/Source/Runtime/Net/Iris/Public/Iris/ReplicationSystem/ReplicationSystem.h
--- a/Engine/Source/Runtime/Net/Iris/Public/Iris/ReplicationSystem/ReplicationSystem.h
+++ b/Engine/Source/Runtime/Net/Iris/Public/Iris/ReplicationSystem/ReplicationSystem.h
@@ -64,6 +64,7 @@ namespace UE::Net
 		class FReplicationSystemImpl;
 		class FReplicationSystemInternal;
 		class FObjectReferenceCache;
+		class FNetStatsContext;
 	}
 }
 
@@ -802,6 +803,11 @@ public:
 	IRISCORE_API void StopParallelPhase();
 	IRISCORE_API bool IsInParallelPhase() const;
 
+	// Used to provide a NetStatsContext for the FNetSerializationContext - thread safe
+	IRISCORE_API UE::Net::Private::FNetStatsContext* AcquireStatsForSerializationContext();
+	// Used to free any aquired NetStatsContext used by this SerializationContext  thread safe
+	IRISCORE_API void ReleaseStatsForSerializationContext(UE::Net::Private::FNetStatsContext* NetStatsContext);
+
 	/** Function to query if a particular replicated object is relevant to a connection or not. */
 	UE_EXPERIMENTAL(5.7, "IsRelevantForConnection() is experimental and should not be used.")
 	IRISCORE_API bool IsRelevantForConnection(FNetRefHandle NetRefHandle, uint32 ConnectionId) const;
@@ -912,5 +918,27 @@ inline UReplicationSystem* GetReplicationSystem(uint32 Id)
 	return Id >= static_cast<uint32>(FReplicationSystemFactory::ReplicationSystems.Num()) ? nullptr : FReplicationSystemFactory::ReplicationSystems[Id];
 }
 
+/**
+* A helper RAII class to allow thread-safe access to FNetStatsContext
+* If created within a ParallelPhase, the constructor acquires a ChildNetStatsContext,
+* otherwise it falls back to the ParentNetStatsContext. This means that we do not
+* lock any mutexes for the duration of this object, only briefly inside AcquireChildNetStatsContext
+*/
+class FScopedNetStatsContext
+{
+public:
+	
+	IRISCORE_API explicit FScopedNetStatsContext(UReplicationSystem* const InReplicationSystem);
+
+	IRISCORE_API ~FScopedNetStatsContext();
+
+	IRISCORE_API Private::FNetStatsContext* GetRawPtr();
+
+private:
+	class Private::FNetStatsContext* StoredNetStatsContext = nullptr;
+	UReplicationSystem* StoredReplicationSystem = nullptr;
+	bool bWasAcquiredInParallelPhase = false;
+};
+
 } // end namespace UE::Net
 
