Iris: Don't kill the server if one connection write fails

#rnx
#Synced-CL 49063813

diff --git a/Engine/Source/Runtime/Net/Iris/Private/Iris/ReplicationSystem/ReplicationWriter.cpp b/Engine/Source/Runtime/Net/Iris/Private/Iris/ReplicationSystem/ReplicationWriter.cpp
--- a/Engine/Source/Runtime/Net/Iris/Private/Iris/ReplicationSystem/ReplicationWriter.cpp
+++ b/Engine/Source/Runtime/Net/Iris/Private/Iris/ReplicationSystem/ReplicationWriter.cpp
@@ -3245,7 +3245,6 @@ FReplicationWriter::EWriteStatus FReplicationWriter::PrepareAndSendHugeObjectPay
 	}
 
 	const bool bEnqueueSuccess = Attachments.Enqueue(AttachmentType, ObjectIndexForOOBAttachment, MakeArrayView(PartialNetBlobs.GetData(), PartialNetBlobs.Num()));
-	check(bEnqueueSuccess);
 	if (!bEnqueueSuccess)
 	{
 		UE_LOG(LogIris, Error, TEXT("Unable to enqueue huge object attachments: %s. Connection %u will be disconnected."), *NetRefHandleManager->PrintObjectFromIndex(InternalIndex), Context.GetLocalConnectionId());
