Iris
- Improved error message for NetBlobAssembler

#Synced-CL 49098511

diff --git a/Engine/Source/Runtime/Net/Iris/Private/Iris/ReplicationSystem/NetBlob/NetBlobAssembler.cpp b/Engine/Source/Runtime/Net/Iris/Private/Iris/ReplicationSystem/NetBlob/NetBlobAssembler.cpp
--- a/Engine/Source/Runtime/Net/Iris/Private/Iris/ReplicationSystem/NetBlob/NetBlobAssembler.cpp
+++ b/Engine/Source/Runtime/Net/Iris/Private/Iris/ReplicationSystem/NetBlob/NetBlobAssembler.cpp
@@ -5,6 +5,7 @@
 #include "Iris/ReplicationSystem/NetBlob/SequentialPartialNetBlobHandler.h"
 #include "Iris/Serialization/NetBitStreamReader.h"
 #include "Iris/Serialization/NetSerializationContext.h"
+#include "Iris/Core/IrisLog.h"
 
 namespace UE::Net
 {
@@ -27,6 +28,7 @@ void FNetBlobAssembler::AddPartialNetBlob(FNetSerializationContext& Context, FNe
 
 	if (!ensureMsgf(PartialNetBlobHandlerConfig != nullptr, TEXT("NetBlobAssembler requires a PartialNetBlobHandlerConfig")))
 	{
+		UE_LOG(LogIris, Error, TEXT("NetBlobAssembler - Requires a PartialNetBlobHandlerConfig"));
 		Context.SetError(GNetError_InvalidValue);
 		return;
 	}
@@ -77,6 +79,8 @@ void FNetBlobAssembler::AddPartialNetBlob(FNetSerializationContext& Context, FNe
 		{
 			bIsBrokenSequence = true;
 			Context.SetError(GNetError_InvalidValue);
+			UE_LOG(LogIris, Error, TEXT("NetBlobAssembler - Expected partCount with first part"), PartCount);
+			ensure(false);
 			return;
 		}
 
@@ -84,6 +88,8 @@ void FNetBlobAssembler::AddPartialNetBlob(FNetSerializationContext& Context, FNe
 		{
 			bIsBrokenSequence = true;
 			Context.SetError(GNetError_InvalidValue);
+			UE_LOG(LogIris, Error, TEXT("NetBlobAssembler - Received invalid PartCount %u configured MaxPartCount %u"), PartCount, PartialNetBlobHandlerConfig->GetMaxPartCount());
+			ensure(false);
 			return;
 		}
 
@@ -93,6 +99,8 @@ void FNetBlobAssembler::AddPartialNetBlob(FNetSerializationContext& Context, FNe
 		{
 			bIsBrokenSequence = true;
 			Context.SetError(GNetError_InvalidValue);
+			UE_LOG(LogIris, Error, TEXT("NetBlobAssembler - Received invalid PayloadBitCount %u configured MaxPartBitCount %u"), PayloadBitCount, PartialNetBlobHandlerConfig->GetMaxPartBitCount());
+			ensure(false);
 			return;
 		}
 
@@ -127,6 +135,8 @@ void FNetBlobAssembler::AddPartialNetBlob(FNetSerializationContext& Context, FNe
 		{
 			bIsBrokenSequence = true;
 			Context.SetError(GNetError_InvalidValue);
+			UE_LOG(LogIris, Error, TEXT("NetBlobAssembler - Received invalid PayloadBitCount %u"), PayloadBitCount);
+			ensure(false);
 			return;
 		}
 
