SceneGraph
* Remove material unit tests that used the VerseAssetReplicator because they are unreliable and soon to be deprecated
* Fix material test that crashed if the client entity wasn't valid

#rnx

#Synced-CL 49371252

diff --git a/Engine/Source/Runtime/Net/Iris/Private/Iris/ReplicationSystem/ObjectReferenceCache.cpp b/Engine/Source/Runtime/Net/Iris/Private/Iris/ReplicationSystem/ObjectReferenceCache.cpp
--- a/Engine/Source/Runtime/Net/Iris/Private/Iris/ReplicationSystem/ObjectReferenceCache.cpp
+++ b/Engine/Source/Runtime/Net/Iris/Private/Iris/ReplicationSystem/ObjectReferenceCache.cpp
@@ -381,6 +381,7 @@ bool FObjectReferenceCache::CreateObjectReferenceInternal(const UObject* Object,
 
 	// Translated from GuidCache, but kept as separate variable to avoid calling the same functions multiple times
 	const bool bIsFullNameStableForNetworking = Object->IsFullNameStableForNetworking();
+	const bool bIsStableForNetworking = Object->IsNameStableForNetworking();
 	const bool bIsDynamic = !bIsFullNameStableForNetworking;
 	const bool bIsStatic = !bIsDynamic;
 
