Iris
* Added new Detach functions in the NetObjectFactory before deprecated the previous ones.
* The new functions will now always be called even if the local object instance is null (not bound) in the replication system
* The new functions now have a clearer EDetachReason enum that tells if and why the object is no longer replicating or if it is just not relevant to this client.
* The new functions are current disabled by default but can be enabled via net.Iris.IsUsingNewDetachAPI

#Synced-CL 49546171

diff --git a/Engine/Source/Runtime/Engine/Private/Net/Iris/ReplicationSystem/NetActorFactory.cpp b/Engine/Source/Runtime/Engine/Private/Net/Iris/ReplicationSystem/NetActorFactory.cpp
--- a/Engine/Source/Runtime/Engine/Private/Net/Iris/ReplicationSystem/NetActorFactory.cpp
+++ b/Engine/Source/Runtime/Engine/Private/Net/Iris/ReplicationSystem/NetActorFactory.cpp
@@ -489,8 +489,15 @@ void UNetActorFactory::SubObjectCreatedFromReplication(UE::Net::FNetRefHandle Ro
 	}
 }
 
+PRAGMA_DISABLE_DEPRECATION_WARNINGS
 void UNetActorFactory::SubObjectDetachedFromReplication(const FDestroyedContext& Context)
 {
+	// DEPRECATED PATH
+	if (UE::Net::IsUsingNewDetachAPI())
+	{
+		return;
+	}
+
 	if (!Context.DestroyedInstance)
 	{
 		return;
@@ -512,6 +519,12 @@ void UNetActorFactory::SubObjectDetachedFromReplication(const FDestroyedContext&
 
 void UNetActorFactory::DetachedFromReplication(const FDestroyedContext& Context)
 {
+	// DEPRECATED PATH
+	if (UE::Net::IsUsingNewDetachAPI())
+	{
+		return;
+	}
+
 	UEngineReplicationBridge* ActorBridge = CastChecked<UEngineReplicationBridge>(Bridge);
 	UNetDriver* NetDriver = ActorBridge->GetNetDriver();
 
@@ -534,6 +547,82 @@ void UNetActorFactory::DetachedFromReplication(const FDestroyedContext& Context)
 		}
 	}	
 }
+PRAGMA_ENABLE_DEPRECATION_WARNINGS
+
+void UNetActorFactory::SubObjectDetachedFromReplication(const FDetachContext& Context, const FSubObjectDetachContext& SubObjectContext)
+{
+	using namespace UE::Net;
+
+	if (!UE::Net::IsUsingNewDetachAPI())
+	{
+		return;
+	}
+
+	if (!Context.DetachedInstance)
+	{
+		return;
+	}
+
+	if (Context.Reason == EDetachReason::TornOff)
+	{
+		return;
+	}
+
+	const bool bIsStable = Context.DetachedInstance->IsNameStableForNetworking();
+	const bool bWillBeDestroyed = !bIsStable;
+
+	if (!bWillBeDestroyed)
+	{
+		return;
+	}
+
+	if (AActor* ActorOwner = Cast<AActor>(SubObjectContext.RootObject))
+	{
+		ActorOwner->OnSubobjectDestroyFromReplication(Context.DetachedInstance);
+	}
+}
+
+void UNetActorFactory::DetachedFromReplication(const FDetachContext& Context, const TOptional<FSubObjectDetachContext>& SubObjectContext)
+{
+	using namespace UE::Net;
+
+	if (!UE::Net::IsUsingNewDetachAPI())
+	{
+		return;
+	}
+
+	AActor* Actor = Cast<AActor>(Context.DetachedInstance);
+	if (!Actor)
+	{
+		// Nothing to do since the actor instance has been detached already
+		return;
+	}
+
+	UEngineReplicationBridge* ActorBridge = CastChecked<UEngineReplicationBridge>(Bridge);
+	UNetDriver* NetDriver = ActorBridge->GetNetDriver();
+
+	const bool bIsDynamic = !Actor->IsNameStableForNetworking();
+	const bool bCanBeDestroyed = (bIsDynamic || Context.Reason == EDetachReason::StaticDestroyed) && NetDriver->ShouldClientDestroyActor(Actor);
+
+	if (Context.Reason == EDetachReason::TornOff)
+	{
+		if (!NetDriver->ShouldClientDestroyTearOffActors())
+		{
+			NetDriver->ClientSetActorTornOff(Actor);
+		}
+		else if (bCanBeDestroyed)
+		{
+			Actor->PreDestroyFromReplication();
+			Actor->Destroy(true);
+		}
+	}
+	else if (bCanBeDestroyed)
+	{
+		// Note that subobjects have already been detached by the ReplicationBridge
+		Actor->PreDestroyFromReplication();
+		Actor->Destroy(true);
+	}
+}
 
 TOptional<UNetActorFactory::FWorldInfoData> UNetActorFactory::GetWorldInfo(const FWorldInfoContext& Context) const
 {
diff --git a/Engine/Source/Runtime/Engine/Private/Net/Iris/ReplicationSystem/NetRootObjectFactory.cpp b/Engine/Source/Runtime/Engine/Private/Net/Iris/ReplicationSystem/NetRootObjectFactory.cpp
--- a/Engine/Source/Runtime/Engine/Private/Net/Iris/ReplicationSystem/NetRootObjectFactory.cpp
+++ b/Engine/Source/Runtime/Engine/Private/Net/Iris/ReplicationSystem/NetRootObjectFactory.cpp
@@ -106,6 +106,7 @@ UNetObjectFactory::FInstantiateResult UNetRootObjectFactory::InstantiateReplicat
 	return Result;
 }
 
+PRAGMA_DISABLE_DEPRECATION_WARNINGS
 void UNetRootObjectFactory::DetachedFromReplication(const FDestroyedContext& Context)
 {
 	UE_LOG(LogIris, Verbose, TEXT("%hs: %s | Reason: %s | Flags: %s"), __FUNCTION__, *GetNameSafe(Context.DestroyedInstance), LexToString(Context.DestroyReason), LexToString(Context.DestroyFlags));
@@ -124,6 +125,35 @@ void UNetRootObjectFactory::DetachedFromReplication(const FDestroyedContext& Con
 		Context.DestroyedInstance->MarkAsGarbage();
 	}
 }
+PRAGMA_ENABLE_DEPRECATION_WARNINGS
+
+void UNetRootObjectFactory::DetachedFromReplication(const FDetachContext& Context, const TOptional<FSubObjectDetachContext>& SubObjectContext)
+{
+	using namespace UE::Net;
+
+	UE_LOG(LogIris, Verbose, TEXT("%hs: %s | Reason: %s"), __FUNCTION__, *GetNameSafe(Context.DetachedInstance), LexToString(Context.Reason));
+
+	if (!Context.DetachedInstance)
+	{
+		return;
+	}
+
+	if (Context.Reason == EDetachReason::TornOff)
+	{
+		return;
+	}
+	
+	const bool bIsStable = Context.DetachedInstance->IsNameStableForNetworking();
+	const bool bCanBeDestroyed = !bIsStable || Context.Reason == EDetachReason::StaticDestroyed;
+	
+	if (bCanBeDestroyed)
+	{
+		Context.DetachedInstance->PreDestroyFromReplication();
+
+		//TODO: Some objects will want to opt-out of being MarkAsGarbage to support verse semantics
+		Context.DetachedInstance->MarkAsGarbage();
+	}
+}
 
 TOptional<UNetObjectFactory::FWorldInfoData> UNetRootObjectFactory::GetWorldInfo(const FWorldInfoContext& Context) const
 {
diff --git a/Engine/Source/Runtime/Engine/Private/Net/Iris/ReplicationSystem/NetSubObjectFactory.cpp b/Engine/Source/Runtime/Engine/Private/Net/Iris/ReplicationSystem/NetSubObjectFactory.cpp
--- a/Engine/Source/Runtime/Engine/Private/Net/Iris/ReplicationSystem/NetSubObjectFactory.cpp
+++ b/Engine/Source/Runtime/Engine/Private/Net/Iris/ReplicationSystem/NetSubObjectFactory.cpp
@@ -232,7 +232,7 @@ UNetSubObjectFactory::FInstantiateResult UNetSubObjectFactory::InstantiateReplic
 		.Template = Template, // Force the template to be used as the default init state source. It's possible GetArchetype() will be different then the template
 	};
 
-	// We must defer call OnSubObjectCreatedFromReplication after the state has been applied to the owning actor in order to behave like old system.
+	// Call OnSubObjectCreatedFromReplication after the state has been applied to the owning actor
 	Result.Flags |= EReplicationBridgeCreateNetRefHandleResultFlags::ShouldCallSubObjectCreatedFromReplication;
 
 	// Created objects may be destroyed.
@@ -258,6 +258,7 @@ void UNetSubObjectFactory::SubObjectCreatedFromReplication(UE::Net::FNetRefHandl
 	ensureMsgf(false, TEXT("NetSubObjectFactory::SubObjectCreatedFromReplication should never be called since subobjects cannot have their own subobject list. RootObject: %s, SubObjectCreated: %s"), *Bridge->PrintObjectFromNetRefHandle(RootObject), *Bridge->PrintObjectFromNetRefHandle(SubObjectCreated));
 }
 
+PRAGMA_DISABLE_DEPRECATION_WARNINGS
 void UNetSubObjectFactory::DetachedFromReplication(const FDestroyedContext& Context)
 {
 	// If the SubObject is being torn off it is up to owning actor to clean it up properly
@@ -280,6 +281,47 @@ void UNetSubObjectFactory::DetachedFromReplication(const FDestroyedContext& Cont
 	Context.DestroyedInstance->PreDestroyFromReplication();
 	Context.DestroyedInstance->MarkAsGarbage();
 }
+PRAGMA_ENABLE_DEPRECATION_WARNINGS
+
+void UNetSubObjectFactory::DetachedFromReplication(const FDetachContext& Context, const TOptional<FSubObjectDetachContext>& SubObjectContext)
+{
+	using namespace UE::Net;
+
+	UObject* DetachedSubObject = Context.DetachedInstance;
+	if (!DetachedSubObject)
+	{
+		// The instance has already been detached, nothing else to do here
+		return;
+	}
+
+	if (Context.Reason == EDetachReason::TornOff)
+	{
+		// If the SubObject is being torn off it is up to owning actor to clean it up properly
+		return;
+	}
+
+	if (!SubObjectContext.IsSet())
+	{
+		ensureMsgf(false, TEXT("The SubObjectContext should always exist for the NetSubObjectFactory"));
+		return;
+	}
+
+	const bool bIsStableSubObject = DetachedSubObject->IsNameStableForNetworking();
+
+	//TODO: Handle StaticDestroyed reason
+	const bool bCanBeDestroyed = !bIsStableSubObject;
+
+	//TODO: Decide if there are conditions where we could do things differently here when looking at the root object detach reason...
+
+	if (!bCanBeDestroyed)
+	{
+		// Nothing to do if the subobject shouldn't be destroyed
+		return;
+	}
+
+	DetachedSubObject->PreDestroyFromReplication();
+	DetachedSubObject->MarkAsGarbage();
+}
 
 TOptional<UNetObjectFactory::FWorldInfoData> UNetSubObjectFactory::GetWorldInfo(const FWorldInfoContext& Context) const
 {
diff --git a/Engine/Source/Runtime/Engine/Public/Net/Iris/ReplicationSystem/NetActorFactory.h b/Engine/Source/Runtime/Engine/Public/Net/Iris/ReplicationSystem/NetActorFactory.h
--- a/Engine/Source/Runtime/Engine/Public/Net/Iris/ReplicationSystem/NetActorFactory.h
+++ b/Engine/Source/Runtime/Engine/Public/Net/Iris/ReplicationSystem/NetActorFactory.h
@@ -187,9 +187,14 @@ public:
 	ENGINE_API virtual void PostInit(const FPostInitContext& Context) override;
 
 	ENGINE_API virtual void SubObjectCreatedFromReplication(UE::Net::FNetRefHandle RootObject, UE::Net::FNetRefHandle SubObjectCreated) override;
-	ENGINE_API virtual void SubObjectDetachedFromReplication(const FDestroyedContext& Context) override;
 
-	ENGINE_API virtual void DetachedFromReplication(const FDestroyedContext& Context) override;
+	PRAGMA_DISABLE_DEPRECATION_WARNINGS
+	virtual void SubObjectDetachedFromReplication(const FDestroyedContext& Context) override;
+	virtual void DetachedFromReplication(const FDestroyedContext& Context) override;
+	PRAGMA_ENABLE_DEPRECATION_WARNINGS
+
+	ENGINE_API virtual void SubObjectDetachedFromReplication(const FDetachContext& Context, const FSubObjectDetachContext& SubObjectContext) override;
+	ENGINE_API virtual void DetachedFromReplication(const FDetachContext& Context, const TOptional<FSubObjectDetachContext>& SubObjectContext) override;
 
 	ENGINE_API virtual TOptional<FWorldInfoData> GetWorldInfo(const FWorldInfoContext& Context) const override;
 
diff --git a/Engine/Source/Runtime/Engine/Public/Net/Iris/ReplicationSystem/NetRootObjectFactory.h b/Engine/Source/Runtime/Engine/Public/Net/Iris/ReplicationSystem/NetRootObjectFactory.h
--- a/Engine/Source/Runtime/Engine/Public/Net/Iris/ReplicationSystem/NetRootObjectFactory.h
+++ b/Engine/Source/Runtime/Engine/Public/Net/Iris/ReplicationSystem/NetRootObjectFactory.h
@@ -63,7 +63,11 @@ public:
 
 	ENGINE_API virtual FInstantiateResult InstantiateReplicatedObjectFromHeader(const FInstantiateContext& Context, const UE::Net::FNetObjectCreationHeader* Header) override;
 
+	PRAGMA_DISABLE_DEPRECATION_WARNINGS
 	ENGINE_API virtual void DetachedFromReplication(const FDestroyedContext& Context) override;
+	PRAGMA_ENABLE_DEPRECATION_WARNINGS
+
+	ENGINE_API virtual void DetachedFromReplication(const FDetachContext& Context, const TOptional<FSubObjectDetachContext>& SubObjectContext) override;
 
 	ENGINE_API virtual TOptional<FWorldInfoData> GetWorldInfo(const FWorldInfoContext& Context) const override;
 
diff --git a/Engine/Source/Runtime/Engine/Public/Net/Iris/ReplicationSystem/NetSubObjectFactory.h b/Engine/Source/Runtime/Engine/Public/Net/Iris/ReplicationSystem/NetSubObjectFactory.h
--- a/Engine/Source/Runtime/Engine/Public/Net/Iris/ReplicationSystem/NetSubObjectFactory.h
+++ b/Engine/Source/Runtime/Engine/Public/Net/Iris/ReplicationSystem/NetSubObjectFactory.h
@@ -89,7 +89,11 @@ public:
 
 	ENGINE_API virtual void SubObjectCreatedFromReplication(UE::Net::FNetRefHandle RootObject, UE::Net::FNetRefHandle SubObjectCreated) override;
 
+	PRAGMA_DISABLE_DEPRECATION_WARNINGS
 	ENGINE_API virtual void DetachedFromReplication(const FDestroyedContext& Context) override;
+	PRAGMA_ENABLE_DEPRECATION_WARNINGS
+
+	ENGINE_API virtual void DetachedFromReplication(const FDetachContext& Context, const TOptional<FSubObjectDetachContext>& SubObjectContext) override;
 
 	ENGINE_API virtual TOptional<FWorldInfoData> GetWorldInfo(const FWorldInfoContext& Context) const override;
 
diff --git a/Engine/Source/Runtime/Net/Iris/Private/Iris/IrisConfig.cpp b/Engine/Source/Runtime/Net/Iris/Private/Iris/IrisConfig.cpp
--- a/Engine/Source/Runtime/Net/Iris/Private/Iris/IrisConfig.cpp
+++ b/Engine/Source/Runtime/Net/Iris/Private/Iris/IrisConfig.cpp
@@ -18,6 +18,9 @@ static FAutoConsoleVariableRef CVarUseIrisReplicationRef(TEXT("net.Iris.UseIrisR
 static bool bScheduleCreationDependenciesFirst = true;
 static FAutoConsoleVariableRef CVarScheduleCreationDependenciesFirst(TEXT("net.Iris.ScheduleCreationDependenciesFirst"), bScheduleCreationDependenciesFirst, TEXT("If enabled we will schedule an objects creation dependencies before the object if they are not confirmed as created."));
 
+static bool bIsUsingNewDetachAPI = false;
+static FAutoConsoleVariableRef CVarIsUsingNewDetachAPI(TEXT("net.Iris.IsUsingNewDetachAPI"), bIsUsingNewDetachAPI, TEXT("If enabled clients will use the new detach API."));
+
 }
 
 bool ShouldUseIrisReplication()
@@ -48,6 +51,11 @@ bool IsSchedulingCreationDependenciesFirst()
 	return Private::CVars::bScheduleCreationDependenciesFirst;
 }
 
+bool IsUsingNewDetachAPI()
+{
+	return Private::CVars::bIsUsingNewDetachAPI;
+}
+
 }
 
 
diff --git a/Engine/Source/Runtime/Net/Iris/Private/Iris/ReplicationSystem/NetObjectFactory.cpp b/Engine/Source/Runtime/Net/Iris/Private/Iris/ReplicationSystem/NetObjectFactory.cpp
--- a/Engine/Source/Runtime/Net/Iris/Private/Iris/ReplicationSystem/NetObjectFactory.cpp
+++ b/Engine/Source/Runtime/Net/Iris/Private/Iris/ReplicationSystem/NetObjectFactory.cpp
@@ -162,3 +162,48 @@ TUniquePtr<UE::Net::FNetObjectCreationHeader> UNetObjectFactory::ReadHeader(UE::
 	return nullptr;
 }
 
+//------------------------------------------------------------------------
+
+const TCHAR* LexToString(UE::Net::EDetachReason Reason)
+{
+	using namespace UE::Net;
+
+	switch (Reason)
+	{
+		case EDetachReason::Stopped:
+		{
+			return TEXT("Stopped");
+		} break;
+
+		case EDetachReason::NoLongerRelevant:
+		{
+			return TEXT("NoLongerRelevant");
+		} break;
+
+		case EDetachReason::StaticDestroyed:
+		{
+			return TEXT("StaticDestroyed");
+		} break;
+
+		case EDetachReason::TornOff:
+		{
+			return TEXT("TornOff");
+		} break;
+
+		case EDetachReason::Migrated:
+		{
+			return TEXT("Migrated");
+		} break;
+
+		case EDetachReason::RootObjectStillRelevant:
+		{
+			return TEXT("RootObjectStillRelevant");
+		} break;
+
+		default:
+		{
+			ensure(false);
+			return TEXT("Missing");
+		} break;
+	}
+}
diff --git a/Engine/Source/Runtime/Net/Iris/Private/Iris/ReplicationSystem/NetRefHandleManager.cpp b/Engine/Source/Runtime/Net/Iris/Private/Iris/ReplicationSystem/NetRefHandleManager.cpp
--- a/Engine/Source/Runtime/Net/Iris/Private/Iris/ReplicationSystem/NetRefHandleManager.cpp
+++ b/Engine/Source/Runtime/Net/Iris/Private/Iris/ReplicationSystem/NetRefHandleManager.cpp
@@ -5,8 +5,6 @@
 #include "Iris/Core/IrisLog.h"
 #include "Iris/Core/IrisProfiler.h"
 
-#include "Iris/ReplicationSystem/NetRefHandleManagerTypes.h"
-
 #include "Iris/Serialization/InternalNetSerializationContext.h"
 #include "Iris/Serialization/NetSerializationContext.h"
 
@@ -42,6 +40,35 @@ static FAutoConsoleVariableRef CVarAllowCreateNetObjectForNetRefHandlePendingDes
 	TEXT("If set to true we allow creation of a new NetObject for a NetRefHandle pending destroy that might still be replicating data to client. This is currently not fully supported and can lead to issues on the client.")
 );
 
+static_assert(GetDetachReasonBitsNeeded() <= 8U, "EInternalDetachReason is a uint8");
+
+const TCHAR* LexToString(EInternalDetachReason DetachReason)
+{
+	switch (DetachReason)
+	{
+		case EInternalDetachReason::Normal:
+			return TEXT("Normal");
+		case EInternalDetachReason::NotRelevant:
+			return TEXT("NotRelevant");
+		case EInternalDetachReason::StaticDestroyed:
+			return TEXT("StaticDestroyed");
+		case EInternalDetachReason::Migrated:
+			return TEXT("Migrated");
+		case EInternalDetachReason::TornOff:
+			return TEXT("TornOff");
+		case EInternalDetachReason::Max:
+		{
+			ensureMsgf(false, TEXT("Max entry should never be used except for calculating the bits needed"));
+			return TEXT("Max");
+		} break;
+		default:
+		{
+			ensure(false);
+			return TEXT("Missing");
+		};
+	}
+}
+
 static bool IsValidIndex(FInternalNetRefIndex InternalIndex)
 {
 	return InternalIndex != FNetRefHandleManager::InvalidInternalIndex;
diff --git a/Engine/Source/Runtime/Net/Iris/Private/Iris/ReplicationSystem/NetRefHandleManager.h b/Engine/Source/Runtime/Net/Iris/Private/Iris/ReplicationSystem/NetRefHandleManager.h
--- a/Engine/Source/Runtime/Net/Iris/Private/Iris/ReplicationSystem/NetRefHandleManager.h
+++ b/Engine/Source/Runtime/Net/Iris/Private/Iris/ReplicationSystem/NetRefHandleManager.h
@@ -18,7 +18,9 @@
 #include "Iris/ReplicationSystem/NetDependencyData.h"
 #include "Iris/ReplicationSystem/NetObjectFactoryRegistry.h"
 #include "Iris/ReplicationSystem/NetRefHandle.h"
+#include "Iris/ReplicationSystem/NetRefHandleManagerTypes.h"
 #include "Iris/ReplicationSystem/ObjectReferenceTypes.h"
+#include "Iris/ReplicationSystem/ReplicationBridgeTypes.h"
 
 #include "Misc/TransactionallySafeCriticalSection.h"
 
@@ -36,8 +38,6 @@ namespace UE::Net
 	namespace Private
 	{
 		class FReplicationProtocolManager;
-
-		struct FCreateNetObjectParams;
 	}
 }
 
@@ -65,6 +65,7 @@ enum class ERemoveDependentObjectFlags : uint32
 ENUM_CLASS_FLAGS(ERemoveDependentObjectFlags);
 
 /** The reason that is replicated to clients when an on object is removed from the replication system. */
+//TODO: Make this deprecated
 enum class ECloseReplicationReason : uint8
 {
 	None = 0,
@@ -74,6 +75,11 @@ enum class ECloseReplicationReason : uint8
 	MaxBits = 1U
 };
 
+consteval uint32 GetCloseReasonBits()
+{
+	return static_cast<uint32>(ECloseReplicationReason::MaxBits);
+}
+
 // Internal class to manage NetHandles and their internal data
 class FNetRefHandleManager
 {
@@ -118,6 +124,8 @@ public:
 		uint16 bAllowDestroyInstanceFromRemote : 1U = 0;
 		uint16 bNeedsFullCopyAndQuantize : 1U = 0;
 		uint16 bWantsFullPoll : 1U = 0;
+		uint16 bIsAsyncStopping : 1U = 0;
+		/** Set to true when the object had StopReplication called directly on it */
 		uint16 bPendingEndReplication : 1U = 0;
 		uint16 bHasCachedCreationInfo : 1U = 0;
 		uint16 bIsDebugObject : 1U = 0;
@@ -128,7 +136,14 @@ public:
 		/** The async loading priority for references collected from this object */
 		EIrisAsyncLoadingPriority IrisAsyncLoadingPriority = EIrisAsyncLoadingPriority::Default;
 
+		//Deprecated reason
 		ECloseReplicationReason CloseReason = ECloseReplicationReason::None;
+
+		/** Stores the detach reason while the object is stopping replication */
+		EInternalDetachReason InternalDetachReason = EInternalDetachReason::Normal;
+
+		/** Returns true if this is a root object */
+		bool IsRootObject() const { return SubObjectRootIndex == InvalidInternalIndex; }
 	
 		/** Returns true if this is a SubObject */
 		bool IsSubObject() const { return SubObjectRootIndex != InvalidInternalIndex; }
diff --git a/Engine/Source/Runtime/Net/Iris/Private/Iris/ReplicationSystem/NetRefHandleManagerTypes.h b/Engine/Source/Runtime/Net/Iris/Private/Iris/ReplicationSystem/NetRefHandleManagerTypes.h
--- a/Engine/Source/Runtime/Net/Iris/Private/Iris/ReplicationSystem/NetRefHandleManagerTypes.h
+++ b/Engine/Source/Runtime/Net/Iris/Private/Iris/ReplicationSystem/NetRefHandleManagerTypes.h
@@ -13,12 +13,37 @@ namespace UE::Net
 namespace UE::Net::Private
 {
 
-	/** Holds important parameters needed to create a NetObject */
-	struct FCreateNetObjectParams
-	{
-		FNetObjectFactoryId NetFactoryId = InvalidNetObjectFactoryId;
-		EIrisAsyncLoadingPriority IrisAsyncLoadingPriority = EIrisAsyncLoadingPriority::Default;
-		const UE::Net::FReplicationProtocol* ReplicationProtocol = nullptr;
-	};
+/** Holds important parameters needed to create a NetObject */
+struct FCreateNetObjectParams
+{
+	FNetObjectFactoryId NetFactoryId = InvalidNetObjectFactoryId;
+	EIrisAsyncLoadingPriority IrisAsyncLoadingPriority = EIrisAsyncLoadingPriority::Default;
+	const UE::Net::FReplicationProtocol* ReplicationProtocol = nullptr;
+};
+
+/** An internal Iris-only enum used to serialize the reason why an object is no longer replicated to a client. */
+enum class EInternalDetachReason : uint8
+{
+	/** Normal is simply a StopReplication call without any explicit reason given. */
+	Normal = 0U,
+	/** When a root object stops replicating to a client because it's no longer relevant. */
+	NotRelevant,
+	/** When a stable replicated object is destroyed even when the client loads it locally. */
+	StaticDestroyed,
+	/** When a replicated object is migrated from one game server to another. */
+	Migrated,
+	/** When a replicated object is torn off by the authority. */
+	TornOff,
+	/** Add new entries above this one. */
+	Max,
+};
+
+const TCHAR* LexToString(EInternalDetachReason DetachReason);
+
+consteval uint32 GetDetachReasonBitsNeeded()
+{
+	return UE::Net::GetBitsNeeded(static_cast<uint32>(EInternalDetachReason::Max) - 1U);
+}
+
 
 }
\ No newline at end of file
diff --git a/Engine/Source/Runtime/Net/Iris/Private/Iris/ReplicationSystem/ObjectReplicationBridge.cpp b/Engine/Source/Runtime/Net/Iris/Private/Iris/ReplicationSystem/ObjectReplicationBridge.cpp
--- a/Engine/Source/Runtime/Net/Iris/Private/Iris/ReplicationSystem/ObjectReplicationBridge.cpp
+++ b/Engine/Source/Runtime/Net/Iris/Private/Iris/ReplicationSystem/ObjectReplicationBridge.cpp
@@ -447,7 +447,7 @@ void UObjectReplicationBridge::RemoveDestructionInfosForGroup(UE::Net::FNetObjec
 					ObjectsToRemove.Add(ObjectData.RefHandle);
 					ObjectIndicesToRemove.Add(InternalObjectIndex);
 				}
-			}			
+			}
 		}
 
 		for (FNetRefHandle Handle : MakeArrayView(ObjectsToRemove))
@@ -507,26 +507,37 @@ void UObjectReplicationBridge::ReadAndExecuteDestructionInfoFromRemote(FReplicat
 
 	UE_NET_TRACE_SCOPE(DestructionInfo, *Context.SerializationContext.GetBitStreamReader(), Context.SerializationContext.GetTraceCollector(), ENetTraceVerbosity::Trace);
 
-	// Destroy instance here, or defer to later?
+	//TODO: We need to support destruction messages that work on any stable objects, not just those with static path names
+
+	// Read the static path referencing the object
 	FNetObjectReference ReferenceToDestroy;
 	ReadFullNetObjectReference(Context.SerializationContext, ReferenceToDestroy);
 
+	// Read the NetFactoryId
 	const FNetObjectFactoryId NetFactoryId = (FNetObjectFactoryId)Context.SerializationContext.GetBitStreamReader()->ReadBits(FNetObjectFactoryRegistry::GetMaxBits());
 
-	// Destroy the reference
-	// Resolve the reference in order to be able to destroy it
-	if (const UObject* Instance = ObjectReferenceCache->ResolveObjectReference(ReferenceToDestroy, Context.SerializationContext.GetInternalContext()->ResolveContext))
+	// Find the object referenced by the path name
+	const UObject* Instance = ObjectReferenceCache->ResolveObjectReference(ReferenceToDestroy, Context.SerializationContext.GetInternalContext()->ResolveContext);
+	if (!Instance)
 	{
-		const FInternalNetRefIndex InternalReplicationIndex = NetRefHandleManager->GetInternalIndex(ReferenceToDestroy.GetRefHandle());
-		if (InternalReplicationIndex != FNetRefHandleManager::InvalidInternalIndex)
-		{
-			NetRefHandleManager->GetReplicatedObjectDataNoCheck(InternalReplicationIndex).bPendingEndReplication = 1U;
-		}
+		// Object does not exist, nothing to destroy
+		return;
+	}
+	
+	// Look if this object has an internal index assigned. If it doesn't have one it's because it was never replicated to the client before
+	const FInternalNetRefIndex InternalReplicationIndex = NetRefHandleManager->GetInternalIndex(ReferenceToDestroy.GetRefHandle());
+	if (InternalReplicationIndex != FNetRefHandleManager::InvalidInternalIndex)
+	{
+		NetRefHandleManager->GetReplicatedObjectDataNoCheck(InternalReplicationIndex).bPendingEndReplication = 1U;
+	}
 
-		UE_LOG(LogIris, VeryVerbose, TEXT("ReadAndExecuteDestructionInfoFromRemote on %s | %s | NetFactoryId: %u"), *ReferenceToDestroy.ToString(), *ReferenceToDestroy.GetRefHandle().ToString(), NetFactoryId);
+	UE_LOG(LogIris, VeryVerbose, TEXT("ReadAndExecuteDestructionInfoFromRemote on %s | %s | NetFactoryId: %u"), *ReferenceToDestroy.ToString(), *ReferenceToDestroy.GetRefHandle().ToString(), NetFactoryId);
 
+	// Old deprecated path
+	if (!UE::Net::IsUsingNewDetachAPI())
+	{
 		constexpr EReplicationBridgeDestroyInstanceReason DestroyReason = EReplicationBridgeDestroyInstanceReason::Destroy;
-		const EReplicationBridgeDestroyInstanceFlags DestroyFlags = EReplicationBridgeDestroyInstanceFlags::AllowDestroyInstanceFromRemote;
+		constexpr EReplicationBridgeDestroyInstanceFlags DestroyFlags = EReplicationBridgeDestroyInstanceFlags::AllowDestroyInstanceFromRemote;
 
 		// if a subobject owner is to be destroyed we want to detach all subobjects before doing so to ensure we execute expected callbacks
 		// We keep tracking them internally
@@ -536,7 +547,8 @@ void UObjectReplicationBridge::ReadAndExecuteDestructionInfoFromRemote(FReplicat
 	}
 	else
 	{
-		//$TODO: Need to validate that it's safe if the resolve returns false
+		constexpr EDetachReason DetachReason = EDetachReason::StaticDestroyed;
+		DetachRootObjectFromRemote(ReferenceToDestroy.GetRefHandle(), DetachReason, NetFactoryId);
 	}
 }
 
@@ -953,9 +965,72 @@ UE::Net::FNetRefHandle UObjectReplicationBridge::StartReplicatingSubObject(UObje
 	return SubObjectRefHandle;
 }
 
+UE::Net::EDetachReason ConvertDetachReasonFromInternalToPublic(UE::Net::Private::EInternalDetachReason InternalDetachReason)
+{
+	using namespace UE::Net;
+	using namespace UE::Net::Private;
+	switch (InternalDetachReason)
+	{
+		case EInternalDetachReason::Normal:
+		{
+			return EDetachReason::Stopped;
+		} break;
+
+		case EInternalDetachReason::NotRelevant:
+		{
+			return EDetachReason::NoLongerRelevant;
+		} break;
+
+		case EInternalDetachReason::StaticDestroyed:
+		{
+			return EDetachReason::StaticDestroyed;
+		} break;
+
+		case EInternalDetachReason::Migrated:
+		{
+			return EDetachReason::Migrated;
+		} break;
+
+		case EInternalDetachReason::TornOff:
+		{
+			return EDetachReason::TornOff;
+		} break;
+
+		default:
+		{
+			ensureMsgf(false, TEXT("Missing entry!"));
+		} break;
+	}
+
+	/* Convert the enum to the old deprecated reason enum if we want to be deprecation friendly
+	if (InternalDetachReason == EIrisDetachReason::Normal || 
+		InternalDetachReason == EIrisDetachReason::NotRelevant)
+	{
+		DestroyReason = EReplicationBridgeDestroyInstanceReason::Destroy;
 
-void UObjectReplicationBridge::DestroyNetObjectFromRemote(UE::Net::Private::FInternalNetRefIndex InternalObjectIndex, EReplicationBridgeDestroyInstanceReason DestroyReason)
+		// First rely on the object to tell if it was stable or not
+		UObject* ObjectInstance = NetRefHandleManager->GetReplicatedObjectInstance(InternalObjectIndex);
+		if (ObjectInstance)
+		{
+			if (ObjectInstance->IsNameStableForNetworking())
+			{
+				DestroyReason = EReplicationBridgeDestroyInstanceReason::DoNotDestroy;
+	  		}
+	 	}
+	}
+	// Fallback on the netrefhandle otherwise for backward compatibility
+	else if (ObjectData.RefHandle.IsStatic())
+	{
+		DestroyReason = EReplicationBridgeDestroyInstanceReason::DoNotDestroy;
+	}
+}*/
+
+	return EDetachReason::Stopped;
+}
+
+void UObjectReplicationBridge::DestroyNetObjectFromRemote(UE::Net::Private::FInternalNetRefIndex InternalObjectIndex, EReplicationBridgeDestroyInstanceReason DestroyReason, UE::Net::Private::EInternalDetachReason InternalDetachReason)
 {
+	using namespace UE::Net;
 	using namespace UE::Net::Private;
 
 	if (InternalObjectIndex == FNetRefHandleManager::InvalidInternalIndex)
@@ -963,14 +1038,34 @@ void UObjectReplicationBridge::DestroyNetObjectFromRemote(UE::Net::Private::FInt
 		return;
 	}
 
+	ensure(InternalDetachReason != EInternalDetachReason::Max);
+
 	FNetRefHandleManager::FReplicatedObjectData& ObjectData = NetRefHandleManager->GetReplicatedObjectDataNoCheck(InternalObjectIndex);
 
-	ObjectData.bTearOff = (DestroyReason == EReplicationBridgeDestroyInstanceReason::TearOff);
+
+	ObjectData.bTearOff = (InternalDetachReason == EInternalDetachReason::TornOff);
+	ensure(!ObjectData.bTearOff || (InternalDetachReason == EInternalDetachReason::TornOff && DestroyReason == EReplicationBridgeDestroyInstanceReason::TearOff));
+
+	//TODO: Find why this is needed here? Is it to block client StopReplicatingNetRefHandle calls ?
 	ObjectData.bPendingEndReplication = 1U;
 
-	const EReplicationBridgeDestroyInstanceFlags DestroyFlags = (ObjectData.bAllowDestroyInstanceFromRemote ? EReplicationBridgeDestroyInstanceFlags::AllowDestroyInstanceFromRemote : EReplicationBridgeDestroyInstanceFlags::None);
+	EReplicationBridgeDestroyInstanceFlags DestroyFlags = EReplicationBridgeDestroyInstanceFlags::None;
+	
+	{
+		// This flag is set by the factory on object instantation.
+		// What it basically boils down to is that the object can be "MarkedForGarbage" when it is no longer replicated. So right now only dynamic objects that the factory created via NewObject.
+		// TODO: Decide if it's useful for the factory to register and keep a backup of objects it newed, and how to expose that trait back again in the detach context. The factory could be able tell which object it newed or not (IsNameStable, etc), but maybe it's less error prone to have it tag objects on creation and expose those same tags back on detachment.
+		if (ObjectData.bAllowDestroyInstanceFromRemote)
+		{
+			// If we considere this object to be destroyable, add the flag again here...
+			EnumAddFlags(DestroyFlags,  EReplicationBridgeDestroyInstanceFlags::AllowDestroyInstanceFromRemote);
+		}
+	}
+
+	// Convert from internal to public reason
+	EDetachReason PublicDetachReason = ConvertDetachReasonFromInternalToPublic(InternalDetachReason);
 
-	UE_LOG_OBJECTREPLICATIONBRIDGE(Verbose, TEXT("DestroyNetObjectFromRemote for %s | DestroyReason: %s | DestroyFlags: %s "), *NetRefHandleManager->PrintObjectFromIndex(InternalObjectIndex), LexToString(DestroyReason), LexToString(DestroyFlags) );
+	UE_LOG_OBJECTREPLICATIONBRIDGE(Verbose, TEXT("DestroyNetObjectFromRemote for %s | InternalReason: %s | PublicReason: %s |  DestroyReason: %s | DestroyFlags: %s "), *NetRefHandleManager->PrintObjectFromIndex(InternalObjectIndex), LexToString(InternalDetachReason), LexToString(PublicDetachReason), LexToString(DestroyReason), LexToString(DestroyFlags) );
 
 	const FNetRefHandle DestroyedObjectHandle = ObjectData.RefHandle;
 	ensureMsgf(DestroyedObjectHandle.IsValid(), TEXT("Object with InternalIndex assigned an invalid handle: %s"), *NetRefHandleManager->PrintObjectFromIndex(InternalObjectIndex));
@@ -979,12 +1074,31 @@ void UObjectReplicationBridge::DestroyNetObjectFromRemote(UE::Net::Private::FInt
 		return;
 	}
 
-	// if a subobject owner is to be destroyed we want to detach all subobjects before doing so to ensure we execute expected callbacks
-	// We keep tracking them internally
-	DetachSubObjectInstancesFromRemote(InternalObjectIndex, DestroyReason, DestroyFlags);
+	if (UE::Net::IsUsingNewDetachAPI() == false)
+	{
+		// if a root object is detaching, we want to detach all subobjects first to ensure we execute expected callbacks
+		// We keep tracking them internally
+		if (ObjectData.IsRootObject())
+		{
+			DetachSubObjectInstancesFromRemote(InternalObjectIndex, DestroyReason, DestroyFlags);
+		}
 		
-	// Allow derived bridges to cleanup any instance info they have stored
-	DetachInstanceFromRemote(DestroyedObjectHandle, DestroyReason, DestroyFlags, ObjectData.NetFactoryId);
+		// Tell factories to take care of the detached object.
+		DetachInstanceFromRemote(DestroyedObjectHandle, DestroyReason, DestroyFlags, ObjectData.NetFactoryId);
+	}
+	else
+	{
+		if (ObjectData.IsRootObject())
+		{
+			//TODO: Test how Iris behaves if a subobject is destroyed with a different reason (tearoff, etc.) from it's root object and we receive both the destruction messages at the same or different time.
+			//It feels like the root object reason migh always have priority over everything else ?
+			DetachRootObjectFromRemote(DestroyedObjectHandle, PublicDetachReason, ObjectData.NetFactoryId);
+		}
+		else
+		{
+			DetachSubObjectFromRemote(DestroyedObjectHandle, PublicDetachReason, ObjectData.NetFactoryId);
+		}
+	}
 
 	// Detach instance protocol
 	InternalDetachInstanceFromNetRefHandle(InternalObjectIndex);
@@ -1214,13 +1328,270 @@ void UObjectReplicationBridge::StopReplicatingNetObject(UObject* Instance, EEndR
 	}
 }
 
+void UObjectReplicationBridge::DetachRootObjectFromRemote(FNetRefHandle RootObjectHandle, UE::Net::EDetachReason DetachReason, UE::Net::FNetObjectFactoryId NetFactoryId)
+{
+	using namespace UE::Net;
+	using namespace UE::Net::Private;
+
+	UE_LOG_OBJECTREPLICATIONBRIDGE(Verbose, TEXT("%hs: %s Reason: %s"), __FUNCTION__, *PrintObjectFromNetRefHandle(RootObjectHandle), LexToString(DetachReason));
+
+	if (!FNetObjectFactoryRegistry::IsValidFactoryId(NetFactoryId))
+	{
+		ensureMsgf(false, TEXT("Received invalid netfactory: %u for %s"), NetFactoryId,  *PrintObjectFromNetRefHandle(RootObjectHandle));
+		return;
+	}
+
+	UNetObjectFactory* RootObjectFactory = GetNetFactory(NetFactoryId);
+	
+	const FInternalNetRefIndex DetachedInternalIndex = NetRefHandleManager->GetInternalIndex(RootObjectHandle);
+
+	// The instance may be null even if the object was still replicating. This mostly happens when the local destroy paths call StopReplication early before waiting on the server to do it.
+	UObject* RootObjectInstance = GetObjectFromReferenceHandle(RootObjectHandle);
+
+	const UNetObjectFactory::FSubObjectDetachContext SubObjectDetachContext
+	{
+		.RootObject = RootObjectInstance,
+		.RootObjectHandle = RootObjectHandle,
+		.RootObjectDetachReason = DetachReason,
+	};
+
+	// Detach all the subobjects
+	for (const FInternalNetRefIndex SubObjectInternalIndex : NetRefHandleManager->GetChildSubObjects(DetachedInternalIndex))
+	{
+		UE_LOG_OBJECTREPLICATIONBRIDGE(Verbose, TEXT("%hs: Subobject %s | Reason: %s"), __FUNCTION__, *NetRefHandleManager->PrintObjectFromIndex(SubObjectInternalIndex), LexToString(DetachReason));
+
+		FNetRefHandleManager::FReplicatedObjectData& SubObjectData = NetRefHandleManager->GetReplicatedObjectDataNoCheck(SubObjectInternalIndex);
+
+		SubObjectData.bTearOff = DetachReason == EDetachReason::TornOff;
+		SubObjectData.bPendingEndReplication = 1U;
+
+		const FNetRefHandle SubObjectHandle = SubObjectData.RefHandle;
+
+		//TODO: This commented code calls the old functions when the new API is enabled.
+		// But I'm not sure its worth the trouble to be 100% backwards compatible vs forcing users to migrate to the new path. Keeping it here before deciding on final design.
+		
+		/*
+		PRAGMA_DISABLE_DEPRECATION_WARNINGS
+		{			
+			EReplicationBridgeDestroyInstanceFlags SubObjectDestroyFlags = EReplicationBridgeDestroyInstanceFlags::AllowDestroyInstanceFromRemote;
+			{
+				// The subobject is allowed to be destroyed if both the owner and the subobject allows it.
+				if (SubObjectHandle.IsStatic() || !SubObjectData.bAllowDestroyInstanceFromRemote)
+				{
+					EnumRemoveFlags(SubObjectDestroyFlags, EReplicationBridgeDestroyInstanceFlags::AllowDestroyInstanceFromRemote);
+				}
+
+				// TODO: need something better here...
+				if (RootObjectHandle.IsStatic() && SubObjectData.bAllowDestroyInstanceFromRemote)
+				{
+					// When ending replication of static objects without destroying them, we should always destroy instances spawned from replication
+					// as they will be recreated when the static object is scoped again.
+					EnumAddFlags(SubObjectDestroyFlags, EReplicationBridgeDestroyInstanceFlags::AllowDestroyInstanceFromRemote);
+				}
+			}
+
+			// The deprecated struct	
+			const UNetObjectFactory::FDestroyedContext SubObjectDestroyContext
+			{
+				.DestroyedInstance = GetObjectFromReferenceHandle(SubObjectHandle),
+				.DestroyedObjectHandle = SubObjectHandle,
+				.RootObject = RootObjectInstance,
+				.RootObjectHandle = RootObjectHandle,
+				.DestroyReason = DestroyReason,
+				.DestroyFlags = SubObjectDestroyFlags,
+			};
+
+			// The deprecated functions would only be called if these conditions are met.
+			const bool bCallDeprecatedFunctions = SubObjectDestroyContext.DestroyedInstance != nullptr &&
+											SubObjectDestroyContext.RootObject != nullptr &&
+											DestroyReason != EReplicationBridgeDestroyInstanceReason::Migrated;
+			if (bCallDeprecatedFunctions)
+			{
+				RootObjectFactory->SubObjectDetachedFromReplication(SubObjectDestroyContext);
+			}
+		}
+		PRAGMA_ENABLE_DEPRECATION_WARNINGS
+		*/
+
+		const UNetObjectFactory::FDetachContext DetachContext
+		{
+			.DetachedInstance = GetObjectFromReferenceHandle(SubObjectHandle),
+			.DetachedObjectHandle = SubObjectHandle,
+			.Reason = DetachReason,
+		};
+
+		// Tell the root object he will lose this subobject
+		RootObjectFactory->SubObjectDetachedFromReplication(DetachContext, SubObjectDetachContext);
+
+		// Now ask the subobject's factory to handle the detachment
+		if (UNetObjectFactory* SubObjectFactory = GetNetFactory(SubObjectData.NetFactoryId))
+		{
+			/*PRAGMA_DISABLE_DEPRECATION_WARNINGS
+			if (bCallDeprecatedFunctions)
+			{
+				SubObjectFactory->DetachedFromReplication(SubObjectDestroyContext);
+			}
+			PRAGMA_ENABLE_DEPRECATION_WARNINGS*/
+
+			SubObjectFactory->DetachedFromReplication(DetachContext, SubObjectDetachContext);
+		}
+		
+		if (DetachReason != EDetachReason::Migrated)
+		{
+			UnregisterInstance(SubObjectHandle);
+		}
+	}
+
+	// Now that subobjects have been detached, handle the root object itself
+
+	// The deprecated function
+	/*
+	PRAGMA_DISABLE_DEPRECATION_WARNINGS
+	{
+		const UNetObjectFactory::FDestroyedContext DestroyContext
+		{
+			.DestroyedInstance = RootObjectInstance,
+			.DestroyedObjectHandle = RootObjectHandle,
+			.DestroyReason = DestroyReason,
+			.DestroyFlags = DestroyFlags,
+		};
+		if (DestroyContext.DestroyedInstance != nullptr &&
+			DestroyReason != EReplicationBridgeDestroyInstanceReason::Migrated)
+		{
+			RootObjectFactory->DetachedFromReplication(DestroyContext);
+		}
+	}
+	PRAGMA_ENABLE_DEPRECATION_WARNINGS
+	*/
+	
+	const UNetObjectFactory::FDetachContext DetachContext
+	{
+		.DetachedInstance = RootObjectInstance,
+		.DetachedObjectHandle = RootObjectHandle,
+		.Reason = DetachReason,
+	};
+	RootObjectFactory->DetachedFromReplication(DetachContext, {});
+
+	// Migrated objects need to keep the Handle->Pointer binding alive since proxies reuse pointers instead of recreating them. This won't be required in the future if proxies simply do not instantiate a local instance and use protocol data directly.
+	if (DetachReason == EDetachReason::Migrated)
+	{
+        return;
+	}
+
+	UnregisterInstance(RootObjectHandle);
+}
+
+void UObjectReplicationBridge::DetachSubObjectFromRemote(FNetRefHandle DetachedSubObjectHandle, UE::Net::EDetachReason DetachReason, UE::Net::FNetObjectFactoryId NetFactoryId)
+{
+	using namespace UE::Net;
+	using namespace UE::Net::Private;
+
+	UE_LOG_OBJECTREPLICATIONBRIDGE(Verbose, TEXT("%hs: %s Reason: %s"), __FUNCTION__, *PrintObjectFromNetRefHandle(DetachedSubObjectHandle), LexToString(DetachReason));
+
+	if (!FNetObjectFactoryRegistry::IsValidFactoryId(NetFactoryId))
+	{
+		ensureMsgf(false, TEXT("Received invalid netfactory: %u for %s"), NetFactoryId,  *PrintObjectFromNetRefHandle(DetachedSubObjectHandle));
+		return;
+	}
+
+	UNetObjectFactory* SubObjectFactory = GetNetFactory(NetFactoryId);
+
+	const FInternalNetRefIndex SubObjectInternalIndex = NetRefHandleManager->GetInternalIndex(DetachedSubObjectHandle);
+
+	const FNetRefHandleManager::FReplicatedObjectData& DetachedSubObjectData = NetRefHandleManager->GetReplicatedObjectData(SubObjectInternalIndex);
+
+	if (!DetachedSubObjectData.IsSubObject())
+	{
+		ensureMsgf(false, TEXT("Detached obbject %s is not a subobject."), *PrintObjectFromNetRefHandle(DetachedSubObjectHandle));
+		return;
+	}
+
+	// The instance may be null even if the object was still replicating. This mostly happens when the root object decides to stop replication before the server requested it.
+	UObject* DetachedSubObjectInstance = GetObjectFromReferenceHandle(DetachedSubObjectHandle);
+	
+	const FInternalNetRefIndex RootObjectIndex = DetachedSubObjectData.SubObjectRootIndex;
+	const FNetRefHandle RootObjectHandle = NetRefHandleManager->GetNetRefHandleFromInternalIndex(RootObjectIndex);
+	UObject* RootObjectInstance = GetObjectFromReferenceHandle(RootObjectHandle);
+	const FNetRefHandleManager::FReplicatedObjectData& RootObjectData = NetRefHandleManager->GetReplicatedObjectData(RootObjectIndex);
+
+	const UNetObjectFactory::FDetachContext DetachContext
+	{
+		.DetachedInstance = DetachedSubObjectInstance,
+		.DetachedObjectHandle = DetachedSubObjectHandle,
+		.Reason = DetachReason,
+	};
+
+	const UNetObjectFactory::FSubObjectDetachContext SubObjectDetachContext
+	{
+		.RootObject = RootObjectInstance,
+		.RootObjectHandle = RootObjectHandle,
+		.RootObjectDetachReason = EDetachReason::RootObjectStillRelevant,
+	};
+
+	// Start by telling the RootObject that one of it's subobject is no longer replicated.
+	if (UNetObjectFactory* RootObjectFactory = GetNetFactory(RootObjectData.NetFactoryId))
+	{
+		/*
+		PRAGMA_DISABLE_DEPRECATION_WARNINGS
+		{
+			EReplicationBridgeDestroyInstanceFlags SubObjectDestroyFlags = EReplicationBridgeDestroyInstanceFlags::None;
+			{
+				// The subobject is allowed to be destroyed if both the owner and the subobject allows it.
+				if (DetachedSubObjectData.bAllowDestroyInstanceFromRemote)
+				{
+					// When ending replication of static objects without destroying them, we should always destroy instances spawned from replication
+					// as they will be recreated when the static object is scoped again.
+					EnumAddFlags(SubObjectDestroyFlags, EReplicationBridgeDestroyInstanceFlags::AllowDestroyInstanceFromRemote);
+				}
+			}
+			const UNetObjectFactory::FDestroyedContext SubObjectDestroyContext
+			{
+				.DestroyedInstance = GetObjectFromReferenceHandle(SubObjectHandle),
+				.DestroyedObjectHandle = SubObjectHandle,
+				.RootObject = RootObjectInstance,
+				.RootObjectHandle = RootObjectHandle,
+				.DestroyReason = DestroyReason,
+				.DestroyFlags = SubObjectDestroyFlags,
+			};
+
+			// The deprecated functions would only be called if those conditions were met.
+			const bool bCallDeprecatedFunctions =
+				SubObjectDestroyContext.DestroyedInstance != nullptr &&
+				SubObjectDestroyContext.RootObject != nullptr &&
+				DestroyReason != EReplicationBridgeDestroyInstanceReason::Migrated;
+			if (bCallDeprecatedFunctions)
+			{
+				RootObjectFactory->SubObjectDetachedFromReplication(SubObjectDestroyContext);
+			}
+		}
+		PRAGMA_ENABLE_DEPRECATION_WARNINGS
+		*/
+
+		RootObjectFactory->SubObjectDetachedFromReplication(DetachContext, SubObjectDetachContext);
+	}
+	
+	SubObjectFactory->DetachedFromReplication(DetachContext, SubObjectDetachContext);
+
+	// Migrated objects need to keep the Handle->Pointer binding alive since proxies reuse pointers instead of recreating them. This won't be required in the future if proxies simply do not instantiate a local instance and use protocol data directly.
+	if (DetachReason == EDetachReason::Migrated)
+	{
+		return;
+	}
+
+	UnregisterInstance(DetachedSubObjectHandle);
+}
+
+
 void UObjectReplicationBridge::DetachInstanceFromRemote(FNetRefHandle Handle, EReplicationBridgeDestroyInstanceReason DestroyReason, EReplicationBridgeDestroyInstanceFlags DestroyFlags, UE::Net::FNetObjectFactoryId NetFactoryId)
 {
 	using namespace UE::Net;
 	using namespace UE::Net::Private;
 
+	ensure(UE::Net::IsUsingNewDetachAPI() == false);
+
 	UE_LOG_OBJECTREPLICATIONBRIDGE(Verbose, TEXT("DetachInstanceFromRemote %s DestroyReason: %s DestroyFlags: %u"), *PrintObjectFromNetRefHandle(Handle), ToCStr(LexToString(DestroyReason)), unsigned(DestroyFlags));
 
+PRAGMA_DISABLE_DEPRECATION_WARNINGS
 	UNetObjectFactory::FDestroyedContext DestroyContext
 	{
 		.DestroyedInstance = GetObjectFromReferenceHandle(Handle),
@@ -1262,6 +1633,7 @@ void UObjectReplicationBridge::DetachInstanceFromRemote(FNetRefHandle Handle, ER
 			}
 		}
 	}
+PRAGMA_ENABLE_DEPRECATION_WARNINGS
 
 	if (DestroyReason != EReplicationBridgeDestroyInstanceReason::Migrated)
 	{
@@ -1273,6 +1645,8 @@ void UObjectReplicationBridge::DetachSubObjectInstancesFromRemote(UE::Net::Priva
 {
 	using namespace UE::Net::Private;
 
+	ensure(UE::Net::IsUsingNewDetachAPI() == false);
+
 	if (InternalObjectIndex == FNetRefHandleManager::InvalidInternalIndex)
 	{
 		return;
diff --git a/Engine/Source/Runtime/Net/Iris/Private/Iris/ReplicationSystem/ReplicationBridge.cpp b/Engine/Source/Runtime/Net/Iris/Private/Iris/ReplicationSystem/ReplicationBridge.cpp
--- a/Engine/Source/Runtime/Net/Iris/Private/Iris/ReplicationSystem/ReplicationBridge.cpp
+++ b/Engine/Source/Runtime/Net/Iris/Private/Iris/ReplicationSystem/ReplicationBridge.cpp
@@ -387,10 +387,15 @@ void UReplicationBridge::InternalRequestAsyncStopReplication(FNetRefHandle Handl
 {
 	using namespace UE::Net::Private;
 
-	// Destroy SubObjects
+	UE_LOG_REPLICATIONBRIDGE(Verbose, TEXT("%hs: %s | EndFlags: %s"), __FUNCTION__, *NetRefHandleManager->PrintObjectFromNetRefHandle(Handle), ToCStr(LexToString(EndReplicationFlags)));
+	
 	FInternalNetRefIndex ObjectInternalIndex = NetRefHandleManager->GetInternalIndex(Handle);
 	if (ObjectInternalIndex != FNetRefHandleManager::InvalidInternalIndex)
 	{
+		FNetRefHandleManager::FReplicatedObjectData& ObjectData = NetRefHandleManager->GetReplicatedObjectDataNoCheck(ObjectInternalIndex);
+
+		ObjectData.bIsAsyncStopping = true;
+
 		AddPendingEndReplication(Handle, EndReplicationFlags);
 
 		// If flush is required we need to wait to after scope update before we initialize async end replication, otherwise we can do it now.
@@ -411,7 +416,8 @@ void UReplicationBridge::InternalRequestAsyncStopReplication(FNetRefHandle Handl
 			// Tag subobject for async stop replication. The check against the scope is needed since the subobjects array might contain destroyed subobjects
 			if (NetRefHandleManager->IsScopableIndex(SubObjectInternalIndex))
 			{
-				UE_LOG_REPLICATIONBRIDGE(Verbose, TEXT("InternalRequestAsyncStopReplicationForSubObjects %s"), *NetRefHandleManager->PrintObjectFromNetRefHandle(SubObjectHandle));
+				//UE_LOG_REPLICATIONBRIDGE(Verbose, TEXT("InternalRequestAsyncStopReplicationForSubObjects %s"), *NetRefHandleManager->PrintObjectFromNetRefHandle(SubObjectHandle));
+				//TODO: Would be faster if we didn't do a recursive call here
 				InternalRequestAsyncStopReplication(SubObjectHandle, EndReplicationFlags);
 			}
 		}
@@ -437,14 +443,20 @@ void UReplicationBridge::StopReplicatingNetRefHandle(FNetRefHandle Handle, EEndR
 			UE_LOG_REPLICATIONBRIDGE(Verbose, TEXT("Ignoring EndReplication called on object already PendingEndReplication %s."), *NetRefHandleManager->PrintObjectFromNetRefHandle(Handle));
 			return;
 		}
-		
+
+		// Store the stop replication reason
+		ObjectData.InternalDetachReason = EInternalDetachReason::Normal;
+
 		if (EnumHasAnyFlags(EndReplicationFlags, EEndReplicationFlags::Migrated))
 		{
 			ObjectData.CloseReason = ECloseReplicationReason::Migrated;
+			ObjectData.InternalDetachReason = EInternalDetachReason::Migrated;
 		}
 
 		if (EnumHasAnyFlags(EndReplicationFlags, EEndReplicationFlags::TearOff))
 		{
+			ObjectData.InternalDetachReason = EInternalDetachReason::TornOff;
+
 			// Add handle to list of objects pending EndReplication indicate that it should be destroyed during next update
 			// We need to do this to cover the case where the torn off object not yet has been added to the scope.
 			AddPendingEndReplication(Handle, EndReplicationFlags);
@@ -458,7 +470,7 @@ void UReplicationBridge::StopReplicatingNetRefHandle(FNetRefHandle Handle, EEndR
 		}
 		else 
 		{
-			// New objects, destroyed during the same frame with posted attachments(RPC):s needs to request a flush to ensure that they get a scope update
+			// New objects, destroyed during the same frame with posted attachments(RPC) need to request a flush to ensure that they get a scope update
 			const UE::Net::Private::FNetBlobManager& NetBlobManager = GetReplicationSystem()->GetReplicationSystemInternal()->GetNetBlobManager();
 			const bool bAllowAutoFlushOfUnProcessedReliableRPCs = bEnableFlushReliableRPCOnDestroy && ObjectData.bNeedsFullCopyAndQuantize;
 			if (bAllowAutoFlushOfUnProcessedReliableRPCs && NetBlobManager.HasUnprocessedReliableAttachments(InternalReplicationIndex))
@@ -467,17 +479,19 @@ void UReplicationBridge::StopReplicatingNetRefHandle(FNetRefHandle Handle, EEndR
 			}
 
 			const bool bFlushRequested = EnumHasAnyFlags(EndReplicationFlags, EEndReplicationFlags::Flush);
-			if (EnumHasAnyFlags(EndReplicationFlags, EEndReplicationFlags::AllowAsyncStopReplication))
+			const bool bAsyncStopRequested = EnumHasAnyFlags(EndReplicationFlags, EEndReplicationFlags::AllowAsyncStopReplication);
+			
+			if (bAsyncStopRequested)
 			{
-				// Defer stop replication until object is no longer scoped by any connections to allow restart replication.
-				AddPendingEndReplication(Handle, EndReplicationFlags);
+				// Defer stop replication until object is no longer scoped by any connections. This allows to restart replication on the object seamlessly if a start was requested for the same object while the objects are still relevant.
+
 
 				if (bFlushRequested)
 				{
 					InternalFlushStateData(Handle);
 				}
 
-				// Note: When using async stop replication We do not detach the instance until it is no longer referenced by an instance.
+				// Note: When using async stop we only detach the instance when it is no longer referenced by every client.
 				InternalRequestAsyncStopReplication(Handle, EndReplicationFlags);
 			}
 			else if (bFlushRequested)
@@ -536,9 +550,9 @@ void UReplicationBridge::StopReplicatingNetRefHandle(FNetRefHandle Handle, EEndR
 void UReplicationBridge::TearOffHandlesPendingTearOff()
 {
 	// Initiate tear off
-	for (FPendingEndReplicationInfo Info : MakeArrayView(HandlesPendingEndReplication))
+	for (const FPendingEndReplicationInfo& Info : MakeArrayView(HandlesPendingEndReplication))
 	{
-		if (EnumHasAnyFlags(Info.DestroyFlags, EEndReplicationFlags::TearOff))
+		if (EnumHasAnyFlags(Info.EndReplicationFlags, EEndReplicationFlags::TearOff))
 		{
 			InternalTearOff(Info.Handle);
 		}
@@ -560,11 +574,11 @@ void UReplicationBridge::UpdateHandlesPendingEndReplication()
 			{
 				FNetRefHandleManager::FReplicatedObjectData& ObjectData = NetRefHandleManager->GetReplicatedObjectDataNoCheck(ObjectInternalIndex);
 				ObjectData.bPendingEndReplication = 1U;
-				DestroyLocalNetHandle(Info.Handle, Info.DestroyFlags);
+				DestroyLocalNetHandle(Info.Handle, Info.EndReplicationFlags);
 			}
 			else
 			{
-				if (EnumHasAnyFlags(Info.DestroyFlags, EEndReplicationFlags::AllowAsyncStopReplication))
+				if (EnumHasAnyFlags(Info.EndReplicationFlags, EEndReplicationFlags::AllowAsyncStopReplication))
 				{
 					FNetRefHandleManager::FReplicatedObjectData& ObjectData = NetRefHandleManager->GetReplicatedObjectDataNoCheck(ObjectInternalIndex);
 					// We keep the object in scope but we mark it a not replicated which will stop replicating it for all clients
@@ -586,7 +600,7 @@ void UReplicationBridge::UpdateHandlesPendingEndReplication()
 				}
 			
 				// Keep object in the pending EndReplication-list until the object is no longer referenced by any ReplicationWriter
-				ObjectsStillPendingEndReplication.Add(FPendingEndReplicationInfo(Info.Handle, Info.DestroyFlags, EPendingEndReplicationImmediate::No));
+				ObjectsStillPendingEndReplication.Add(FPendingEndReplicationInfo(Info.Handle, Info.EndReplicationFlags, EPendingEndReplicationImmediate::No));
 			}
 		}
 	}
@@ -622,7 +636,7 @@ bool UReplicationBridge::RestartOrStopObjectPendingAsyncStopReplication(FNetRefH
 	bool bRestarted = false;
 
 	const int32 Index = HandlesPendingEndReplication.FindLastByPredicate([Handle](const FPendingEndReplicationInfo& Info){ return Info.Handle == Handle; });
-	if (Index != INDEX_NONE && EnumHasAnyFlags(HandlesPendingEndReplication[Index].DestroyFlags, EEndReplicationFlags::AllowAsyncStopReplication))
+	if (Index != INDEX_NONE && EnumHasAnyFlags(HandlesPendingEndReplication[Index].EndReplicationFlags, EEndReplicationFlags::AllowAsyncStopReplication))
 	{
 		UE_LOG_REPLICATIONBRIDGE(Verbose, TEXT("RestartOrStopObjectPendingAsyncStopReplication: %s ForceStop: %u"), *PrintObjectFromNetRefHandle(Handle), bForceStopReplication ? 1U : 0U);
 
@@ -635,11 +649,13 @@ bool UReplicationBridge::RestartOrStopObjectPendingAsyncStopReplication(FNetRefH
 			if (bForceStopReplication)
 			{
 				ObjectData.bPendingEndReplication = 1U;
-				DestroyLocalNetHandle(HandlesPendingEndReplication[Index].Handle, HandlesPendingEndReplication[Index].DestroyFlags);
+				DestroyLocalNetHandle(HandlesPendingEndReplication[Index].Handle, HandlesPendingEndReplication[Index].EndReplicationFlags);
 				bRestarted = false;
 			}
 			else
 			{
+				ObjectData.bIsAsyncStopping = false;
+
 				bRestarted = true;
 				if (NetRefHandleManager->IsDestroyedStartupObject(InternalObjectIndex))
 				{
diff --git a/Engine/Source/Runtime/Net/Iris/Private/Iris/ReplicationSystem/ReplicationReader.cpp b/Engine/Source/Runtime/Net/Iris/Private/Iris/ReplicationSystem/ReplicationReader.cpp
--- a/Engine/Source/Runtime/Net/Iris/Private/Iris/ReplicationSystem/ReplicationReader.cpp
+++ b/Engine/Source/Runtime/Net/Iris/Private/Iris/ReplicationSystem/ReplicationReader.cpp
@@ -6,6 +6,7 @@
 
 #include "HAL/IConsoleManager.h"
 
+#include "Iris/Core/BitTwiddling.h"
 #include "Iris/Core/IrisDebugging.h"
 #include "Iris/Core/IrisLog.h"
 #include "Iris/Core/IrisProfiler.h"
@@ -201,7 +202,7 @@ struct FReplicationReader::FDispatchObjectInfo
 	bool bIsInitialState : 1 = false;
 	bool bHasState : 1 = false;
 	bool bHasAttachments : 1 = false;
-	bool bDestroy : 1 = false;
+	bool bDestroy : 1 = false; // deprecated: only kept for backward compatibility for now when IsUsingNewDetachAPI() == false
 	bool bTearOff : 1 = false;
 	
 	/** When true we will call EndReplication on the object after reading all the bunch */
@@ -216,7 +217,11 @@ struct FReplicationReader::FDispatchObjectInfo
 	/** When true we'll tell the RootObject's factory that this subobject was instantiated */
 	bool bShouldCallSubObjectCreatedFromReplication : 1 = false;
 
+	// deprecated: only kept for backward compatibility for now when IsUsingNewDetachAPI() == false
 	ECloseReplicationReason CloseReason = ECloseReplicationReason::None;
+
+	/** The reason we received about why this object is no longer replicated */
+	EInternalDetachReason InternalDetachReason = EInternalDetachReason::Normal;
 };
 
 //------------------------------------------------------------------------
@@ -230,8 +235,7 @@ bool ReadCloseReason(FNetSerializationContext& Context, ECloseReplicationReason&
 	const bool bHasDestroyReason = Reader.ReadBool();
 	if (bHasDestroyReason)
 	{
-		const uint32 MaxBits = static_cast<uint32>(ECloseReplicationReason::MaxBits);
-		OutReason = static_cast<ECloseReplicationReason>(Reader.ReadBits(MaxBits));
+		OutReason = static_cast<ECloseReplicationReason>(Reader.ReadBits(GetCloseReasonBits()));
 	}
 
 	return !Context.HasErrorOrOverflow();
@@ -486,6 +490,8 @@ uint32 FReplicationReader::ReadRootObjectsPendingDestroy(FNetSerializationContex
 		UE_NET_TRACE_NAMED_OBJECT_SCOPE(DestroyedObjectScope, FNetRefHandle::GetInvalid(), Reader, Context.GetTraceCollector(), ENetTraceVerbosity::Trace);
 
 		FNetRefHandle IncompleteHandle = ReadNetRefHandleId(Context, Reader);
+
+		//===+ Deprecated path for IsUsingNewDetachAPI() == false
 		const bool bShouldDestroyInstance = Reader.ReadBool();
 		if (Context.HasErrorOrOverflow())
 		{
@@ -497,10 +503,18 @@ uint32 FReplicationReader::ReadRootObjectsPendingDestroy(FNetSerializationContex
 		{
 			break;
 		}
+		//===-
+
+		const EInternalDetachReason DetachReason = (EInternalDetachReason)Context.GetBitStreamReader()->ReadBits(GetDetachReasonBitsNeeded());
+		
+		if (Context.HasErrorOrOverflow())
+		{
+			break;
+		}
 
 		if (FPendingBatchData* PendingBatchData = PendingBatchHolder.Find(IncompleteHandle))
 		{
-			EnqueueEndReplication(PendingBatchData, bShouldDestroyInstance, IncompleteHandle, CloseReason);
+			EnqueueEndReplication(PendingBatchData, bShouldDestroyInstance, IncompleteHandle, CloseReason, DetachReason);
 			continue;
 		}
 
@@ -524,6 +538,9 @@ uint32 FReplicationReader::ReadRootObjectsPendingDestroy(FNetSerializationContex
 				Info.bHasAttachments = false;
 				Info.bShouldCallSubObjectCreatedFromReplication = false;
 				Info.CloseReason = CloseReason;
+				Info.InternalDetachReason = DetachReason;
+
+				UE_LOG(LogIris, Verbose, TEXT("%hs deferred end replication for RootObject: %s"), __FUNCTION__, *NetRefHandleManager->PrintObjectFromIndex(InternalIndex));
 
 				// Mark for dispatch
 				ObjectsToDispatchArray->CommitPendingDispatchObjectInfo();
@@ -592,7 +609,7 @@ void FReplicationReader::CleanupObjectData(FReplicatedObjectInfo& ObjectInfo)
 	}
 }
 
-void FReplicationReader::EndReplication(FInternalNetRefIndex InternalIndex, bool bTearOff, bool bDestroyInstance, ECloseReplicationReason CloseReason)
+void FReplicationReader::EndReplication(FInternalNetRefIndex InternalIndex, bool bTearOff, bool bDestroyInstance, ECloseReplicationReason CloseReason, EInternalDetachReason InternalDetachReason)
 {
 	if (InternalIndex == ObjectIndexForOOBAttachment)
 	{
@@ -623,7 +640,7 @@ void FReplicationReader::EndReplication(FInternalNetRefIndex InternalIndex, bool
 		DestroyReason = EReplicationBridgeDestroyInstanceReason::Destroy;
 	}
 		
-	FBridgeSerialization::DestroyNetObjectFromRemote(ReplicationBridge, InternalIndex, DestroyReason);
+	FBridgeSerialization::DestroyNetObjectFromRemote(ReplicationBridge, InternalIndex, DestroyReason, InternalDetachReason);
 
 	CleanupObjectData(*ObjectInfo);
 
@@ -1209,6 +1226,7 @@ void FReplicationReader::ReadObjectInBatch(FNetSerializationContext& Context, FN
 	uint32 NewBaselineIndex = FDeltaCompressionBaselineManager::InvalidBaselineIndex;
 
 	const bool bIsInitialState = bHasState && Reader.ReadBool();
+	
 	bool bShouldCallSubObjectCreatedFromReplication = false;
 	uint32 InternalIndex = ObjectIndexForOOBAttachment;
 
@@ -1368,9 +1386,15 @@ void FReplicationReader::ReadObjectInBatch(FNetSerializationContext& Context, FN
 		if (Info.bTearOff)
 		{
 			Info.bWantsDeferredEndReplication = true;
+
+			// TornOff objects don't go through the destruction flow on the sending side.
+			// So no detach reason was serialized in the writer.
+			Info.InternalDetachReason = EInternalDetachReason::TornOff;
 		}
 		else if (bIsForSubObject)
 		{
+			// TODO: Find real subobject reason
+			Info.InternalDetachReason = EInternalDetachReason::Normal;
 			if (bDestroySubObject)
 			{
 				Info.bWantsImmediateEndReplication = true;
@@ -1995,7 +2019,7 @@ void FReplicationReader::DispatchStateData(FNetSerializationContext& Context)
 			FDispatchObjectInfo& Info = *PostDispatchObjectInfo.Info;
 			if (Info.bWantsImmediateEndReplication)
 			{
-				EndReplication(Info.InternalIndex, Info.bTearOff, Info.bDestroy, Info.CloseReason);
+				EndReplication(Info.InternalIndex, Info.bTearOff, Info.bDestroy, Info.CloseReason, Info.InternalDetachReason);
 			}
 		}
 
@@ -2398,7 +2422,7 @@ void FReplicationReader::DispatchEndReplication(FNetSerializationContext& Contex
 		if (Info.bWantsDeferredEndReplication)
 		{
 			// Detach and destroy object
-			EndReplication(Info.InternalIndex, Info.bTearOff, Info.bDestroy, Info.CloseReason);
+			EndReplication(Info.InternalIndex, Info.bTearOff, Info.bDestroy, Info.CloseReason, Info.InternalDetachReason);
 
 			if (Info.bCallEndReplicationOnSubObjects)
 			{
@@ -2407,7 +2431,7 @@ void FReplicationReader::DispatchEndReplication(FNetSerializationContext& Contex
 					// End replication for all subobjects and destroy dynamic ones
 					const FNetRefHandle SubObjectRefHandle = NetRefHandleManager->GetNetRefHandleFromInternalIndex(SubObjectIndex);
 					const bool bShouldDestroySubObject = SubObjectRefHandle.IsDynamic();
-					EndReplication(SubObjectIndex, Info.bTearOff, bShouldDestroySubObject, Info.CloseReason);
+					EndReplication(SubObjectIndex, Info.bTearOff, bShouldDestroySubObject, Info.CloseReason, Info.InternalDetachReason);
 				}
 			}
 		}
@@ -2509,19 +2533,20 @@ void FReplicationReader::ProcessHugeObjectAttachment(FNetSerializationContext& C
 	}
 }
 
-bool FReplicationReader::EnqueueEndReplication(FPendingBatchData* PendingBatchData, bool bShouldDestroyInstance, FNetRefHandle NetRefHandleToEndReplication, ECloseReplicationReason CloseReason)
+bool FReplicationReader::EnqueueEndReplication(FPendingBatchData* PendingBatchData, bool bShouldDestroyInstance, FNetRefHandle NetRefHandleToEndReplication, ECloseReplicationReason CloseReason, EInternalDetachReason DetachReason)
 {
-	UE_LOG(LogIris, Verbose, TEXT("FReplicationReader::EnqueueEndReplication for %s since %s has queued batches"), *NetRefHandleToEndReplication.ToString(), *PendingBatchData->Owner.ToString());
+	UE_LOG(LogIris, Verbose, TEXT("FReplicationReader::EnqueueEndReplication for %s since %s has queued batches | Reason: %s"), *NetRefHandleToEndReplication.ToString(), *PendingBatchData->Owner.ToString(), LexToString(DetachReason));
 
 	/* Template of the extra data we serialize in the chunk.
 	struct FEndReplicationData
 	{
 		FNetRefHandle EndReplicationHandle;
 		bool bShouldDestroyInstance : 1;
-		ECloseReplicationReason Reason : ECloseReplicationReason::MaxBits;
+		ECloseReplicationReason Reason : GetCloseReasonBits;
+		EInternalDetachReason Reason : GetDetachReasonBitsNeeded;
 	}
 	*/
-	constexpr uint32 MaxNumDataBits = (sizeof(FNetRefHandle) * 8U) + 1U /*bShouldDestroyInstance*/ + (uint32)ECloseReplicationReason::MaxBits;
+	constexpr uint32 MaxNumDataBits = (sizeof(FNetRefHandle) * 8U) + 1U /*bShouldDestroyInstance*/ + GetCloseReasonBits() + GetDetachReasonBitsNeeded();
 	constexpr uint32 NumDataWords = (MaxNumDataBits + 31U) / 32U;
 
 	static_assert(MaxNumDataBits < (1U << 30 /*bitsize of FQueuedDataChunk::NumBits*/), "Bits needed don't fit inside FQueuedDataChunk::NumBits");
@@ -2544,8 +2569,8 @@ bool FReplicationReader::EnqueueEndReplication(FPendingBatchData* PendingBatchDa
 	WriteUint64(&Writer, NetRefHandleToEndReplication.GetId());
 	Writer.WriteBool(bShouldDestroyInstance);
 	
-	const uint32 MaxBits = static_cast<uint32>(ECloseReplicationReason::MaxBits);
-	Writer.WriteBits(static_cast<uint32>(CloseReason), MaxBits);
+	Writer.WriteBits(static_cast<uint32>(CloseReason), GetCloseReasonBits());
+	Writer.WriteBits((uint32)DetachReason, GetDetachReasonBitsNeeded());
 	
 	Writer.CommitWrites();
 	ensure(Writer.GetPosBits() <= MaxNumDataBits);
@@ -2632,7 +2657,8 @@ void FReplicationReader::ProcessQueuedBatches()
 				{
 					FNetRefHandle EndReplicationHandle;
 					bool bShouldDestroyInstance : 1;
-					ECloseReplicationReason CloseReason : ECloseReplicationReason::MaxBits;
+				 	ECloseReplicationReason CloseReason : GetCloseReasonBits;
+				 	EInternalDetachReason DetachReason : GetDetachReasonBitsNeeded;
 				}
 				*/
 
@@ -2641,29 +2667,29 @@ void FReplicationReader::ProcessQueuedBatches()
 
 				const bool bShouldDestroyInstance = Reader.ReadBool();
 
-				const uint32 MaxReasonBits = (uint32)ECloseReplicationReason::MaxBits;
-				ECloseReplicationReason CloseReason = (ECloseReplicationReason)Reader.ReadBits(MaxReasonBits);
+				ECloseReplicationReason CloseReason = (ECloseReplicationReason)Reader.ReadBits(GetCloseReasonBits());
 
+				const EInternalDetachReason DetachReason = (EInternalDetachReason)Reader.ReadBits(GetDetachReasonBitsNeeded());
+				
 				if (Reader.IsOverflown())
 				{
 					ensureMsgf(false, TEXT("Error while reading end replication data. Should not happen"));
 					break;
 				}
-				
+
 				// End replication for object
 				const uint32 InternalIndex = NetRefHandleManager->GetInternalIndex(NetRefHandleToEndReplication);
 				if (InternalIndex != FNetRefHandleManager::InvalidInternalIndex)
 				{
 					constexpr bool bTearOff = false;
-
-					EndReplication(InternalIndex, bTearOff, bShouldDestroyInstance, CloseReason);
+					EndReplication(InternalIndex, bTearOff, bShouldDestroyInstance, CloseReason, DetachReason);
 
 					for (const FInternalNetRefIndex SubObjectIndex : NetRefHandleManager->GetSubObjects(InternalIndex))
 					{
 						// End replication for all subobjects and destroy dynamic ones
 						const FNetRefHandle SubObjectRefHandle = NetRefHandleManager->GetNetRefHandleFromInternalIndex(SubObjectIndex);
 						const bool bShouldDestroySubObject = SubObjectRefHandle.IsDynamic();
-						EndReplication(SubObjectIndex, bTearOff, bShouldDestroySubObject, CloseReason);
+						EndReplication(SubObjectIndex, bTearOff, bShouldDestroySubObject, CloseReason, DetachReason);
 					}
 				}
 
diff --git a/Engine/Source/Runtime/Net/Iris/Private/Iris/ReplicationSystem/ReplicationReader.h b/Engine/Source/Runtime/Net/Iris/Private/Iris/ReplicationSystem/ReplicationReader.h
--- a/Engine/Source/Runtime/Net/Iris/Private/Iris/ReplicationSystem/ReplicationReader.h
+++ b/Engine/Source/Runtime/Net/Iris/Private/Iris/ReplicationSystem/ReplicationReader.h
@@ -32,6 +32,8 @@ namespace UE::Net
 		class FResolveAndCollectUnresolvedAndResolvedReferenceCollector;
 		class FNetBlobHandlerManager;
 		class FNetRefHandleManager;
+
+		enum class EInternalDetachReason : uint8;
 	}
 }
 
@@ -187,7 +189,7 @@ private:
 	FReplicatedObjectInfo& StartReplication(uint32 InternalIndex);
 
 	// Remove tracking info for the object with InternalIndex
-	void EndReplication(FInternalNetRefIndex InternalIndex, bool bTearOff, bool bDestroyInstance, ECloseReplicationReason RemovalFlags = ECloseReplicationReason::None);
+	void EndReplication(FInternalNetRefIndex InternalIndex, bool bTearOff, bool bDestroyInstance, ECloseReplicationReason RemovalFlags, EInternalDetachReason DetachReason);
 
 	// Free any data allocated per object
 	void CleanupObjectData(FReplicatedObjectInfo& ObjectInfo);
@@ -222,7 +224,7 @@ private:
 
 	// If we are queuing data for a batch we must also defer calls to EndReplication
 	// This method writes this method in the form of a QueuedChunk
-	bool EnqueueEndReplication(FPendingBatchData* PendingBatchData, bool bShouldDestroyInstance, FNetRefHandle NetRefHandleToEndReplication, ECloseReplicationReason CloseReason);
+	bool EnqueueEndReplication(FPendingBatchData* PendingBatchData, bool bShouldDestroyInstance, FNetRefHandle NetRefHandleToEndReplication, ECloseReplicationReason CloseReason,  EInternalDetachReason DetachReason);
 
 	// Remove a handle from the hot and cold unresolved caches used by ResolveAndDispatchUnresolvedReferences(). If this handle is marked as unresolved
 	// again, it will be added to the hot cache.
diff --git a/Engine/Source/Runtime/Net/Iris/Private/Iris/ReplicationSystem/ReplicationWriter.cpp b/Engine/Source/Runtime/Net/Iris/Private/Iris/ReplicationSystem/ReplicationWriter.cpp
--- a/Engine/Source/Runtime/Net/Iris/Private/Iris/ReplicationSystem/ReplicationWriter.cpp
+++ b/Engine/Source/Runtime/Net/Iris/Private/Iris/ReplicationSystem/ReplicationWriter.cpp
@@ -14,10 +14,12 @@
 #include "Iris/PacketControl/PacketNotification.h"
 
 #include "Iris/ReplicationSystem/DeltaCompression/DeltaCompressionBaselineManager.h"
+
 #include "Iris/ReplicationSystem/Filtering/ReplicationFiltering.h"
+
 #include "Iris/ReplicationSystem/NetBlob/NetObjectBlobHandler.h"
 #include "Iris/ReplicationSystem/NetBlob/PartialNetObjectAttachmentHandler.h"
-#include "Iris/ReplicationSystem/NetRefHandleManager.h"
+
 #include "Iris/ReplicationSystem/ReplicationOperations.h"
 #include "Iris/ReplicationSystem/ReplicationOperationsInternal.h"
 #include "Iris/ReplicationSystem/ReplicationSystem.h"
@@ -2217,16 +2219,42 @@ uint32 FReplicationWriter::WriteRootObjectsPendingDestroy(FNetSerializationConte
 		// Write handle with the needed bitCount
 		WriteNetRefHandleId(Context, ObjectData.RefHandle);
 
-		// Write bit indicating if the static instance should be destroyed or not (could skip the bit for dynamic objects)
-		const bool bShouldDestroyInstance = ObjectData.RefHandle.IsDynamic() || NetRefHandleManager->IsDestroyedStartupObject(InternalIndex);
-		Writer.WriteBool(bShouldDestroyInstance);
+		// Deprecated data needed for IsUsingNewDetachAPI() == false
+		{
+			// Write bit indicating if the static instance should be destroyed or not (could skip the bit for dynamic objects)
+			const bool bShouldDestroyInstance = ObjectData.RefHandle.IsDynamic() || NetRefHandleManager->IsDestroyedStartupObject(InternalIndex);
+			Writer.WriteBool(bShouldDestroyInstance);
+
+			const bool bHasDestroyReason = ObjectData.CloseReason != ECloseReplicationReason::None;
+			Writer.WriteBool(bHasDestroyReason);
+			if (bHasDestroyReason)
+			{
+				Writer.WriteBits(static_cast<uint32>(ObjectData.CloseReason), GetCloseReasonBits());
+			}
+		}
+
 
-		const bool bHasDestroyReason = ObjectData.CloseReason != ECloseReplicationReason::None;
-		Writer.WriteBool(bHasDestroyReason);
-		if (bHasDestroyReason)
+		EInternalDetachReason Reason = ObjectData.InternalDetachReason;
+
+		// New path for IsUsingNewDetachAPI() == true
 		{
-			const uint32 MaxBits = static_cast<uint32>(ECloseReplicationReason::MaxBits);
-			Writer.WriteBits(static_cast<uint32>(ObjectData.CloseReason), MaxBits);
+			// Look if the object is stopping replication for an explicit reason
+			if (NetRefHandleManager->IsDestroyedStartupObject(InternalIndex))
+			{
+				Reason = EInternalDetachReason::StaticDestroyed;
+			}
+			else if (ObjectData.bTearOff)
+			{
+				// The object data reason should have been set already
+				ensure(Reason == EInternalDetachReason::TornOff);
+			}
+			// If the object is still replicating, then it's detaching because it's no longer relevant to this client
+			else if (!ObjectData.bIsAsyncStopping && !ObjectData.bPendingEndReplication)
+			{
+				Reason = EInternalDetachReason::NotRelevant;
+			}
+
+			Writer.WriteBits((uint32)Reason, GetDetachReasonBitsNeeded());
 		}
 
 		if (Writer.IsOverflown())
diff --git a/Engine/Source/Runtime/Net/Iris/Private/Iris/ReplicationSystem/ReplicationWriter.h b/Engine/Source/Runtime/Net/Iris/Private/Iris/ReplicationSystem/ReplicationWriter.h
--- a/Engine/Source/Runtime/Net/Iris/Private/Iris/ReplicationSystem/ReplicationWriter.h
+++ b/Engine/Source/Runtime/Net/Iris/Private/Iris/ReplicationSystem/ReplicationWriter.h
@@ -3,24 +3,32 @@
 #pragma once
 
 // This class will never be included from public headers
-#include "Net/Core/NetBitArray.h"
-#include "Iris/ReplicationSystem/NetRefHandleManager.h"
-#include "Iris/ReplicationSystem/ReplicationTypes.h"
-#include "Iris/ReplicationSystem/ReplicationRecord.h"
+
+#include "Containers/Array.h"
+#include "Containers/List.h"
+#include "Containers/Map.h"
+#include "Containers/Set.h"
+
+#include "Iris/Core/BitTwiddling.h"
+
 #include "Iris/DataStream/DataStream.h"
+
 #include "Iris/ReplicationSystem/AttachmentReplication.h"
 #include "Iris/ReplicationSystem/NetBlob/NetBlob.h"
-#include "Iris/ReplicationSystem/NetRefHandle.h"
 #include "Iris/ReplicationSystem/NetExports.h"
+#include "Iris/ReplicationSystem/NetRefHandle.h"
+#include "Iris/ReplicationSystem/NetRefHandleManager.h"
 #include "Iris/ReplicationSystem/ObjectReferenceCache.h"
 #include "Iris/ReplicationSystem/ReplicationDataStreamDebug.h"
+#include "Iris/ReplicationSystem/ReplicationRecord.h"
+#include "Iris/ReplicationSystem/ReplicationTypes.h"
+
 #include "Iris/Stats/NetStats.h"
-#include "Containers/Array.h"
-#include "Containers/Map.h"
-#include "Containers/List.h"
-#include "Containers/Set.h"
+
 #include "Misc/EnumClassFlags.h"
 
+#include "Net/Core/NetBitArray.h"
+
 // Forward declaration
 class UReplicationSystem;
 class UNetObjectBlobHandler;
@@ -91,7 +99,7 @@ public:
 		Count
 	};
 
-	static_assert((uint8)(EReplicatedObjectState::Count) <= 32, "EReplicatedObjectState must fit in 5 bits. See FReplicationInfo::State and FReplicationRecord::FRecordInfo::ReplicatedObjectState members.");
+	static_assert(UE::Net::GetBitsNeeded(EReplicatedObjectState::Count) <= 5, "EReplicatedObjectState must fit in 5 bits. See FReplicationInfo::State and FReplicationRecord::FRecordInfo::ReplicatedObjectState members.");
 
 	static const TCHAR* LexToString(const EReplicatedObjectState State);
 
diff --git a/Engine/Source/Runtime/Net/Iris/Public/Iris/IrisConfig.h b/Engine/Source/Runtime/Net/Iris/Public/Iris/IrisConfig.h
--- a/Engine/Source/Runtime/Net/Iris/Public/Iris/IrisConfig.h
+++ b/Engine/Source/Runtime/Net/Iris/Public/Iris/IrisConfig.h
@@ -23,6 +23,9 @@ IRISCORE_API EReplicationSystem GetUseIrisReplicationCmdlineValue();
 UE_EXPERIMENTAL(5.7, "Experimental and subject for removal.")
 IRISCORE_API bool IsSchedulingCreationDependenciesFirst();
 
+UE_EXPERIMENTAL(5.8, "Only exposed for CVAR rollback")
+IRISCORE_API bool IsUsingNewDetachAPI();
+
 }
 
 /* NetBitStreamReader/Writer validation support */
diff --git a/Engine/Source/Runtime/Net/Iris/Public/Iris/ReplicationSystem/NetObjectFactory.h b/Engine/Source/Runtime/Net/Iris/Public/Iris/ReplicationSystem/NetObjectFactory.h
--- a/Engine/Source/Runtime/Net/Iris/Public/Iris/ReplicationSystem/NetObjectFactory.h
+++ b/Engine/Source/Runtime/Net/Iris/Public/Iris/ReplicationSystem/NetObjectFactory.h
@@ -28,6 +28,30 @@ namespace UE::Net
 namespace UE::Net
 {
 
+enum class EDetachReason : uint32
+{
+	/** StopReplication called without a specific reason */
+	Stopped = 0U,
+
+	/** Object is still replicating but is no longer relevant to this connection */
+	NoLongerRelevant,
+
+	/** Object is being torn off permanently */
+	TornOff,
+
+	/** Object is migrating */
+	Migrated,
+
+	/** Special entry used in FSubObjectDetachContext when the root object stays relevant. */
+	RootObjectStillRelevant,
+
+	/** Special info sent for static (loaded) objects that are destroyed on the server. Even if never replicated they can ask the RepSystem to synchronize their destruction on clients. */
+	StaticDestroyed,
+
+	/** Put new entries above this one*/
+	MAX,
+};
+
 /** Contextual information passed to an header so it can serialize/deserialize itself  */
 struct FCreationHeaderContext
 {
@@ -90,6 +114,8 @@ public:
 	struct FPostInstantiationContext;
 	struct FPostInitContext;
 	struct FDestroyedContext;
+	struct FDetachContext;
+	struct FSubObjectDetachContext;
 	struct FWorldInfoContext;
 	struct FWorldInfoData;
 
@@ -139,13 +165,17 @@ public:
 	 */
 	virtual void PostInit(const FPostInitContext& Context) {}
 
+	//UE_DEPRECATED(5.8, "FDestroyedContext struct has been deprecated and replaced with the FDetachContext")
+	virtual void DetachedFromReplication(const FDestroyedContext& Context) {}
+
 	/**
-	 * Callback triggered when a replicated object is no longer relevant on the client.
+	 * Callback triggered when a replicated object is no longer replicated on the client.
 	 * This is where a factory would destroy dynamic objects, reset stable objects or put objects back in a pool.
 	 * Note that this callback is called on remotes (clients) only.
-	 * @param Context Gives access to the instance that should be destroyed along with details on what type of destruction is requested.
+	 * @param Context Gives access to the instance that is detached along with details on why the object is no longer replicated
+	 * @param SubObjectContext Optional structure passed only when the detached object is a subobject. Gives access to the root object of the subobject.
 	 */
-	virtual void DetachedFromReplication(const FDestroyedContext& Context) PURE_VIRTUAL(UNetObjectFactory::DetachedFromReplication, );
+	virtual void DetachedFromReplication(const FDetachContext& Context, const TOptional<FSubObjectDetachContext>& SubObjectContext) PURE_VIRTUAL(UNetObjectFactory::DetachedFromReplication,);
 
 	/** 
 	 * Optional callback triggered when a root object managed by this factory gets assigned a dynamic subobject. 
@@ -154,16 +184,19 @@ public:
 	 * @param RootObject The root object that owns the subobject
 	 * @param SubObjectCreated The subobject that was just instantiated
 	*/
-
 	virtual void SubObjectCreatedFromReplication(UE::Net::FNetRefHandle RootObject, UE::Net::FNetRefHandle SubObjectCreated) {}
 
+	//UE_DEPRECATED(5.8, "FDestroyedContext struct has been deprecated and replaced with the FDetachContext")
+	virtual void SubObjectDetachedFromReplication(const FDestroyedContext& Context) {}
+
 	/**
 	 * Optional callback triggered when a subobject will be detached from the replication system and potentially destroyed
-	 * by it's factory in DestroyReplicatedObject. Both static and dynamic subobjects will be passed to this function.
+	 * by it's factory. Both static and dynamic subobjects will be passed to this function.
 	 * Note that this callback is called on remotes (clients) only.
-	 * @param Context Gives access to the subobject about to be destroyed and the root object that owns him.
+	 * @param Context Gives access to the subobject that is no longer replicated.
+	 * @param SubObjectContext Gives access to the rootobject of the detached subobject
 	 */
-	virtual void SubObjectDetachedFromReplication(const FDestroyedContext& Context) {}
+	virtual void SubObjectDetachedFromReplication(const FDetachContext& Context, const FSubObjectDetachContext& SubObjectContext) {}
 
 	/**
 	 * Fetch world information about a replicated object so it can be updated in the network engine.
@@ -228,7 +261,7 @@ public:
 	};
 
 	/** Contextual information for the destroy callbacks */
-	struct FDestroyedContext
+	struct /*UE_DEPRECATED(5.8, "Replaced by the FDetachContext structure")*/ FDestroyedContext
 	{
 		/** The object about to be destroyed */
 		UObject* DestroyedInstance = nullptr;
@@ -245,6 +278,32 @@ public:
 		EReplicationBridgeDestroyInstanceReason DestroyReason = EReplicationBridgeDestroyInstanceReason::DoNotDestroy;
 		EReplicationBridgeDestroyInstanceFlags DestroyFlags = EReplicationBridgeDestroyInstanceFlags::None;
 	};
+	
+	/** Contextual information about an object that is no longer replicated to a client */
+	struct FDetachContext
+	{
+		/** The object that is no longer replicated. Can be null if the object was asked to stop replicating locally before the server authority sent a detach request.  */
+		UObject* DetachedInstance = nullptr;
+
+		/** The handle of the object */
+		UE::Net::FNetRefHandle DetachedObjectHandle;
+
+		/** The context that caused the object to be detached for this client. */
+		UE::Net::EDetachReason Reason = UE::Net::EDetachReason::Stopped;
+	};
+
+	/** Contextual information about the RootObject of a detached subobject. */
+	struct FSubObjectDetachContext
+	{
+		/** The root object instance of the detached subobject. */
+		UObject* RootObject = nullptr;
+
+		/** The root object handle of the detached subobject. */
+		UE::Net::FNetRefHandle RootObjectHandle;
+
+		/** The context in which the root object was detached. Will be set to RootObjectStillRelevant when only the subobject is no longer replicating. */
+		UE::Net::EDetachReason RootObjectDetachReason = UE::Net::EDetachReason::Stopped;
+	};
 
 	/** Details which info needs to be updated in GetWorldInfo */
 	enum class EWorldInfoRequested : uint32
@@ -319,3 +378,5 @@ protected:
 };
 
 ENUM_CLASS_FLAGS(UNetObjectFactory::EWorldInfoRequested);
+
+IRISCORE_API const TCHAR* LexToString(UE::Net::EDetachReason Reason);
diff --git a/Engine/Source/Runtime/Net/Iris/Public/Iris/ReplicationSystem/ObjectReplicationBridge.h b/Engine/Source/Runtime/Net/Iris/Public/Iris/ReplicationSystem/ObjectReplicationBridge.h
--- a/Engine/Source/Runtime/Net/Iris/Public/Iris/ReplicationSystem/ObjectReplicationBridge.h
+++ b/Engine/Source/Runtime/Net/Iris/Public/Iris/ReplicationSystem/ObjectReplicationBridge.h
@@ -40,6 +40,9 @@ namespace UE::Net
 	namespace Private
 	{
 		typedef uint32 FInternalNetRefIndex;
+
+		enum class EInternalDetachReason : uint8;
+
 		class FObjectPollFrequencyLimiter;
 		class FObjectPoller;
 		class FReplicationSystemImpl;
@@ -467,16 +470,36 @@ private:
 
 	void RegisterRemoteInstance(FNetRefHandle RefHandle, UObject* InstancePtr, const UE::Net::FReplicationProtocol* Protocol, UE::Net::FReplicationInstanceProtocol* InstanceProtocol, uint32 ConnectionId);
 
+	/** 
+	 * Called when the authority tells a client that a root object is no longer replicated to this client.
+	 * Will remove the root object and all it's subobjects from the replication system locally and ask the factory to deal with the instance of the object.
+	 * @param RootObjectHandle The handle of the root object
+	 * @param DetachReason The reason why the object is no longer replicated
+	 * @param NetFactoryId The netobjectfactory of the detached object
+	 */
+	void DetachRootObjectFromRemote(FNetRefHandle RootObjectHandle, UE::Net::EDetachReason DetachReason, UE::Net::FNetObjectFactoryId NetFactoryId);
+
+	/**
+	 * Called when the authority tells a client that a specific subobject is no longer replicated to this client.
+	 * @param SubObjectHandle The handle of the sub object
+	 * @param DetachReason The reason why the object is no longer replicated
+	 * @param NetFactoryId The netobjectfactory of the detached object
+	 */
+	void DetachSubObjectFromRemote(FNetRefHandle SubObjectHandle, UE::Net::EDetachReason DetachReason, UE::Net::FNetObjectFactoryId NetFactoryId);
+
+	
 	/**
 	 * Called when the instance is detached from the protocol on request by the remote. 
 	 * @param Handle The handle of the object to destroy or tear off.
 	 * @param DestroyReason Reason for destroying the instance. 
 	 * @param DestroyFlags Special flags such as whether the instance may be destroyed when reason is TearOff, which may revert to destroying, or Destroy.
 	 */
+	//DEPRECATED: Will be replaced with DetachRootObjectFromRemote / DetachSubObjectFromRemote soon
 	void DetachInstanceFromRemote(FNetRefHandle Handle, EReplicationBridgeDestroyInstanceReason DestroyReason, EReplicationBridgeDestroyInstanceFlags DestroyFlags,  UE::Net::FNetObjectFactoryId NetFactoryId);
 
 	/** Detach a subobject on the client due to a request from the authority */
-	void DetachSubObjectInstancesFromRemote(UE::Net::Private::FInternalNetRefIndex InternalObjectIndex, EReplicationBridgeDestroyInstanceReason DestroyReason, EReplicationBridgeDestroyInstanceFlags DestroyFlags);
+	//DEPRECATED: Will be replaced with DetachRootObjectFromRemote soon
+	void DetachSubObjectInstancesFromRemote(UE::Net::Private::FInternalNetRefIndex RootObjectIndex, EReplicationBridgeDestroyInstanceReason DestroyReason, EReplicationBridgeDestroyInstanceFlags DestroyFlags);
 
 	void SetNetPushIdOnInstance(UE::Net::FReplicationInstanceProtocol* InstanceProtocol, FNetHandle NetHandle);
 
@@ -553,7 +576,7 @@ private:
 	 * @param InternalObjectIndex The index of the object to destroy
 	 * @DestroyReason How the object should be destroyed
 	 */
-	void DestroyNetObjectFromRemote(UE::Net::Private::FInternalNetRefIndex InternalObjectIndex, EReplicationBridgeDestroyInstanceReason DestroyReason);
+	void DestroyNetObjectFromRemote(UE::Net::Private::FInternalNetRefIndex InternalObjectIndex, EReplicationBridgeDestroyInstanceReason DestroyReason, UE::Net::Private::EInternalDetachReason InternalDetachReason);
 
 	void ReadAndExecuteDestructionInfoFromRemote(FReplicationBridgeSerializationContext& Context);
 
@@ -748,7 +771,10 @@ struct FAsyncLoadingSimulator
 /** Helper class to handle private functionnalities needed by Iris serialization classes */
 class FBridgeSerialization
 {
-public:
+private:
+
+	friend class FReplicationWriter;
+	friend class FReplicationReader;
 
 	static bool WriteNetRefHandleCreationInfo(UObjectReplicationBridge* Bridge, FReplicationBridgeSerializationContext& Context, FNetRefHandle Handle)
 	{
@@ -760,9 +786,9 @@ public:
 		return Bridge->GetStaticDestructionInfo(Handle);
 	}
 
-	static void DestroyNetObjectFromRemote(UObjectReplicationBridge* Bridge, FInternalNetRefIndex InternalObjectIndex, EReplicationBridgeDestroyInstanceReason DestroyReason)
+	static void DestroyNetObjectFromRemote(UObjectReplicationBridge* Bridge, FInternalNetRefIndex InternalObjectIndex, EReplicationBridgeDestroyInstanceReason DestroyReason, EInternalDetachReason InternalDetachReason)
 	{
-		Bridge->DestroyNetObjectFromRemote(InternalObjectIndex, DestroyReason);
+		Bridge->DestroyNetObjectFromRemote(InternalObjectIndex, DestroyReason, InternalDetachReason);
 	}
 
 	static void ReadAndExecuteDestructionInfoFromRemote(UObjectReplicationBridge* Bridge, FReplicationBridgeSerializationContext& Context)
diff --git a/Engine/Source/Runtime/Net/Iris/Public/Iris/ReplicationSystem/ReplicationBridge.h b/Engine/Source/Runtime/Net/Iris/Public/Iris/ReplicationSystem/ReplicationBridge.h
--- a/Engine/Source/Runtime/Net/Iris/Public/Iris/ReplicationSystem/ReplicationBridge.h
+++ b/Engine/Source/Runtime/Net/Iris/Public/Iris/ReplicationSystem/ReplicationBridge.h
@@ -306,8 +306,8 @@ private:
 	// Destroy all SubObjects owned by provided handle
 	void InternalDestroySubObjects(FNetRefHandle OwnerHandle, EEndReplicationFlags Flags);
 
-	// Internal method to stop replication by filtering out net oobject and once it is no longer replicating to any client it will be destroyed.
-	void InternalRequestAsyncStopReplication(FNetRefHandle OwnerHandle, EEndReplicationFlags Flags);
+	// Internal method to stop replication by filtering out the replicated object. Then once it is no longer replicating to any client it will be destroyed.
+	void InternalRequestAsyncStopReplication(FNetRefHandle Handle, EEndReplicationFlags Flags);
 
 	/**
 	 * Called from ReplicationSystem when a streaming container is about to unload.
@@ -354,10 +354,10 @@ private:
 
 	struct FPendingEndReplicationInfo
 	{
-		FPendingEndReplicationInfo(FNetRefHandle InHandle, EEndReplicationFlags InDestroyFlags, EPendingEndReplicationImmediate InImmediate) : Handle(InHandle), DestroyFlags(InDestroyFlags), Immediate(InImmediate) {}
+		FPendingEndReplicationInfo(FNetRefHandle InHandle, EEndReplicationFlags InEndReplicationFlags, EPendingEndReplicationImmediate InImmediate) : Handle(InHandle), EndReplicationFlags(InEndReplicationFlags), Immediate(InImmediate) {}
 
 		FNetRefHandle Handle;
-		EEndReplicationFlags DestroyFlags;
+		EEndReplicationFlags EndReplicationFlags;
 		EPendingEndReplicationImmediate Immediate;
 	};
 	TArray<FPendingEndReplicationInfo> HandlesPendingEndReplication;
diff --git a/Engine/Source/Runtime/Net/Iris/Public/Iris/ReplicationSystem/ReplicationBridgeTypes.h b/Engine/Source/Runtime/Net/Iris/Public/Iris/ReplicationSystem/ReplicationBridgeTypes.h
--- a/Engine/Source/Runtime/Net/Iris/Public/Iris/ReplicationSystem/ReplicationBridgeTypes.h
+++ b/Engine/Source/Runtime/Net/Iris/Public/Iris/ReplicationSystem/ReplicationBridgeTypes.h
@@ -8,6 +8,7 @@
 
 //------------------------------------------------------------------------
 
+/** Flags controlling how the repsystem should deal with the object that will no longer be replicated. */
 enum class EEndReplicationFlags : uint32
 {
 	None								= 0U,
@@ -54,8 +55,9 @@ struct FReplicationBridgeCreateNetRefHandleResult
 	EReplicationBridgeCreateNetRefHandleResultFlags Flags = EReplicationBridgeCreateNetRefHandleResultFlags::None;
 };
 
-//------------------------------------------------------------------------
 
+//------------------------------------------------------------------------
+// Will be deprecated soon
 enum class EReplicationBridgeDestroyInstanceReason : uint32
 {
 	DoNotDestroy,
@@ -66,7 +68,7 @@ enum class EReplicationBridgeDestroyInstanceReason : uint32
 IRISCORE_API const TCHAR* LexToString(EReplicationBridgeDestroyInstanceReason Reason);
 
 //------------------------------------------------------------------------
-
+// Will be deprecated soon
 enum class EReplicationBridgeDestroyInstanceFlags : uint32
 {
 	None = 0U,
