Fixing compile errors when building with remote object handles

#Synced-CL 49620103

diff --git a/Engine/Source/Runtime/Net/Iris/Public/Iris/ReplicationSystem/ObjectReplicationBridge.h b/Engine/Source/Runtime/Net/Iris/Public/Iris/ReplicationSystem/ObjectReplicationBridge.h
--- a/Engine/Source/Runtime/Net/Iris/Public/Iris/ReplicationSystem/ObjectReplicationBridge.h
+++ b/Engine/Source/Runtime/Net/Iris/Public/Iris/ReplicationSystem/ObjectReplicationBridge.h
@@ -775,6 +775,7 @@ private:
 
 	friend class FReplicationWriter;
 	friend class FReplicationReader;
+	friend class UNetObjectFactory;
 
 	static bool WriteNetRefHandleCreationInfo(UObjectReplicationBridge* Bridge, FReplicationBridgeSerializationContext& Context, FNetRefHandle Handle)
 	{
