Fix AsyncLoading Replication Blocked errors being fired prematurely on some platforms due to difference between Timestamp and Cycles

#rnx
#tests Built and tested against Win64 Client and Dedicated Server with Cooked Data

diff --git a/Engine/Source/Runtime/Net/Iris/Private/Iris/ReplicationSystem/ReplicationReader.cpp b/Engine/Source/Runtime/Net/Iris/Private/Iris/ReplicationSystem/ReplicationReader.cpp
--- a/Engine/Source/Runtime/Net/Iris/Private/Iris/ReplicationSystem/ReplicationReader.cpp
+++ b/Engine/Source/Runtime/Net/Iris/Private/Iris/ReplicationSystem/ReplicationReader.cpp
@@ -2809,7 +2809,7 @@ void FReplicationReader::ProcessQueuedBatches()
 
 			if ((QueuedBatchTimeoutWarningInterval > 0) && (PendingBatchData->PendingBatchTryProcessCount > QueuedBatchTimeoutWarningInterval) && FPlatformProperties::RequiresCookedData() && !FPlatformMisc::IsDebuggerPresent())
 			{
-				const double CurrentTime = FPlatformTime::Seconds();
+				const double CurrentTime = FPlatformTime::ToSeconds64(FPlatformTime::Cycles64());// Different to FPlatformTime::Seconds on some platforms due to GetSecondsTimeOffset
 
 				// Log a warning and tell the server about the blockage every X seconds (first after 30secs, again every 10secs after)
 				if (CurrentTime >= PendingBatchData->NextWarningTimeout)
