Iris
- Implemented support to allow start and StopReplication being called multiple times for static objects including support for flush.
- Refactored FlushLogic to be more robust and behaves closer to normal destroy.
- Fixes issue where the same static NetRefHandle/Object could end up being replicated as multiple NetObjects/Channels causing bugs on client.
- Fixed case where flush was not propagated properly to new connections.

#Synced-CL 49669520

diff --git a/Engine/Plugins/Runtime/ReplicationSystemTestPlugin/Source/Private/Tests/ReplicationSystem/TestFlushBeforeDestroy.cpp b/Engine/Plugins/Runtime/ReplicationSystemTestPlugin/Source/Private/Tests/ReplicationSystem/TestFlushBeforeDestroy.cpp
--- a/Engine/Plugins/Runtime/ReplicationSystemTestPlugin/Source/Private/Tests/ReplicationSystem/TestFlushBeforeDestroy.cpp
+++ b/Engine/Plugins/Runtime/ReplicationSystemTestPlugin/Source/Private/Tests/ReplicationSystem/TestFlushBeforeDestroy.cpp
@@ -15,6 +15,19 @@ namespace UE::Net::Private
 
 class FTestFlushBeforeDestroyFixture : public FPartialNetBlobTestFixture
 {
+public:
+
+	void EnqueueAttachmentToClients(const TArrayView<FReplicationSystemTestClient*const>& TargetClients, FNetRefHandle TargetNetRefHandle)
+	{
+		constexpr uint32 PayloadBitCount = 24;
+		const TRefCountPtr<FNetObjectAttachment>& Attachment = MockNetObjectAttachmentHandler->CreateReliableNetObjectAttachment(PayloadBitCount);
+		FNetObjectReference AttachmentTarget = FObjectReferenceCache::MakeNetObjectReference(TargetNetRefHandle);
+
+		for (FReplicationSystemTestClient* Client : TargetClients)
+		{
+			Server->GetReplicationSystem()->QueueNetObjectAttachment(Client->ConnectionIdOnServer, AttachmentTarget, Attachment);
+		}
+	}
 };
 
 UE_NET_TEST_FIXTURE(FTestFlushBeforeDestroyFixture, TestReliableAttachmentFlushedBeforeDestroy)
@@ -84,6 +97,69 @@ UE_NET_TEST_FIXTURE(FTestFlushBeforeDestroyFixture, TestObjectCreatedAndDestroye
 	UE_NET_ASSERT_EQ(Client->GetReplicationBridge()->GetReplicatedObject(ObjectHandle), nullptr);
 }
 
+UE_NET_TEST_FIXTURE(FTestFlushBeforeDestroyFixture, TestObjectCreatedAndDestroyedSameFrameReplicatesIfFlushedWithPacketLoss)
+{
+	FReplicationSystemTestClient* Client = CreateClient();
+	RegisterNetBlobHandlers(Client);
+
+	// Create and start to replicate object
+	UReplicatedTestObject* ServerObject = Server->CreateObject(0, 0);
+	FNetRefHandle ObjectHandle = ServerObject->NetRefHandle;
+
+	// Destroy object indicating that it should be flushed that is that the final state should be replicated to all clients with the object in scope, this invalidates the creationinfo which has to be cached in order for this to work.
+	Server->DestroyObject(ServerObject, EEndReplicationFlags::Destroy | EEndReplicationFlags::Flush);
+
+	// Send update, dropping data
+	Server->UpdateAndSend({Client}, false);
+
+	// Send update, it should send the data.
+	Server->UpdateAndSend({Client});
+
+	// Verify that object is created
+	UE_NET_ASSERT_NE(Client->GetReplicationBridge()->GetReplicatedObject(ObjectHandle), nullptr);
+
+	// Deliver a packet, make sure that object is destroyed on the client.
+	Server->UpdateAndSend({Client});
+
+	// Verify that object is destroyed
+	UE_NET_ASSERT_EQ(Client->GetReplicationBridge()->GetReplicatedObject(ObjectHandle), nullptr);
+}
+
+UE_NET_TEST_FIXTURE(FTestFlushBeforeDestroyFixture, TestObjectAddedToScopeReplicatesIfFlushed)
+{
+	FReplicationSystemTestClient* Client = CreateClient();
+	RegisterNetBlobHandlers(Client);
+
+	// Create and start to replicate object
+	UReplicatedTestObject* ServerObject = Server->CreateObject(0, 0);
+	FNetRefHandle ObjectHandle = ServerObject->NetRefHandle;
+
+	// Destroy object indicating that it should be flushed that is that the final state should be replicated to all clients with the object in scope, this invalidates the creationinfo which has to be cached in order for this to work.
+	Server->DestroyObject(ServerObject, EEndReplicationFlags::Destroy | EEndReplicationFlags::Flush);
+
+	// Update without sending, should add the object marked for flush to the clients scope
+	Server->NetUpdate();
+	Server->PostSendUpdate();
+
+	// Create another client
+	FReplicationSystemTestClient* Client2 = CreateClient();
+	RegisterNetBlobHandlers(Client2);
+
+	// Send update, we expect the flushed object to be sent to the first client but not the one that joined after the flushed object was removed from scope.
+	Server->UpdateAndSend({Client, Client2});
+
+	// Verify that object is created on client but not on client2
+	UE_NET_ASSERT_NE(Client->GetReplicationBridge()->GetReplicatedObject(ObjectHandle), nullptr);
+	UE_NET_ASSERT_EQ(Client2->GetReplicationBridge()->GetReplicatedObject(ObjectHandle), nullptr);
+
+	// Deliver a packet, make sure that object is destroyed on the client.
+	Server->UpdateAndSend({Client, Client2});
+
+	// Verify that object does not exist on any of the clients
+	UE_NET_ASSERT_EQ(Client->GetReplicationBridge()->GetReplicatedObject(ObjectHandle), nullptr);
+	UE_NET_ASSERT_EQ(Client2->GetReplicationBridge()->GetReplicatedObject(ObjectHandle), nullptr);
+}
+
 UE_NET_TEST_FIXTURE(FTestFlushBeforeDestroyFixture, TestObjectAndSubObjectCreatedAndDestroyedSameFrameReplicatesIfFlushed)
 {
 	FReplicationSystemTestClient* Client = CreateClient();
@@ -1642,7 +1718,6 @@ UE_NET_TEST_FIXTURE(FTestFlushBeforeDestroyFixture, TestTearOffAndFilterOutInSam
 {
 	// Uncomment the verbosiy override scopes and enable UE_NET_ENABLE_REPLICATIONWRITER_LOG to track down issues if test fails.
 	//LOG_SCOPE_VERBOSITY_OVERRIDE(LogIris, ELogVerbosity::Verbose);
-	//LOG_SCOPE_VERBOSITY_OVERRIDE(LogIrisBridge, ELogVerbosity::Verbose);
 
 	FTestEnsureScope EnsureScope;
 
@@ -1823,7 +1898,6 @@ UE_NET_TEST_FIXTURE(FTestFlushBeforeDestroyFixture, TestTearOffOfObjectNotInScop
 {
 	// Uncomment the verbosiy override scopes and enable UE_NET_ENABLE_REPLICATIONWRITER_LOG to track down issues if test fails.
 	//LOG_SCOPE_VERBOSITY_OVERRIDE(LogIris, ELogVerbosity::Verbose);
-	//LOG_SCOPE_VERBOSITY_OVERRIDE(LogIrisBridge, ELogVerbosity::Verbose);
 
 	FTestEnsureScope EnsureScope;
 
@@ -1878,9 +1952,9 @@ UE_NET_TEST_FIXTURE(FTestFlushBeforeDestroyFixture, TestTearOffOfObjectNotInScop
 	Server->SendTo(Client);
 	Server->PostSendUpdate();
 
-	// At this point we're expecting Root in WaitOnFlush, SubObject0 in PendingDestroy and SubObject1 in WaitOnCreateConfirmation.
+	// At this point we're expecting Root in WaitOnFlush due to going out of scope, SubObject0 in PendingDestroy and SubObject1 in WaitOnCreateConfirmation.
 	
-	// Tear off 
+	// Tear off while not in scope with unacked data in flight
 	Server->ReplicationSystem->TearOffNextUpdate(ServerObject->NetRefHandle);
 
 	Server->NetUpdate();
@@ -1897,10 +1971,9 @@ UE_NET_TEST_FIXTURE(FTestFlushBeforeDestroyFixture, TestTearOffOfObjectNotInScop
 		}
 	}
 
-	// Send update
+	// Send and deliver a couple of packets
+	Server->UpdateAndSend({Client});
 	Server->UpdateAndSend({Client});
-
-	// Send update
 	Server->UpdateAndSend({Client});
 
 	// Verify all objects are unresolvable on the client
@@ -2120,4 +2193,445 @@ UE_NET_TEST_FIXTURE(FTestFlushBeforeDestroyFixture, TestSubObjectStateIsNotSentW
 	UE_NET_ASSERT_FALSE(Client->IsValidNetRefHandle(SubObjectHandle));
 }
 
+UE_NET_TEST_FIXTURE(FTestFlushBeforeDestroyFixture, TestDynamicStartFlushStartReplicatesAsNewObject)
+{
+	LOG_SCOPE_VERBOSITY_OVERRIDE(LogIris, ELogVerbosity::Verbose);
+
+	FReplicationSystemTestClient* Client = CreateClient();
+	RegisterNetBlobHandlers(Client);
+
+	// Create and start to replicate object
+	UReplicatedTestObject* ServerObject = Server->CreateObject(0, 0);
+	FNetRefHandle OriginalObjectHandle = ServerObject->NetRefHandle;
+
+	// Stop replicating object using flush
+	Server->ReplicationBridge->EndReplication(ServerObject, EEndReplicationFlags::Destroy | EEndReplicationFlags::Flush);
+
+	// Should not be resolvable by handle
+	UE_NET_ASSERT_EQ(Server->GetReplicationBridge()->GetReplicatedObject(OriginalObjectHandle), nullptr);
+
+	// Since this is a dynamic object, we expect a new handle to be assigned.
+	Server->ReplicationBridge->BeginReplication(ServerObject);
+
+	// Should get new handle
+	UE_NET_ASSERT_TRUE(ServerObject->NetRefHandle != OriginalObjectHandle);
+
+	// New handle should be resolvable 
+	UE_NET_ASSERT_NE(Server->GetReplicationBridge()->GetReplicatedObject(ServerObject->NetRefHandle), nullptr);
+
+	// Send update, it should send the data.
+	Server->UpdateAndSend({Client});
+
+	// Verify that first instance of object is created
+	UE_NET_ASSERT_NE(Client->GetReplicationBridge()->GetReplicatedObject(OriginalObjectHandle), nullptr);
+	// Verify that new instance of object is created 
+	UE_NET_ASSERT_NE(Client->GetReplicationBridge()->GetReplicatedObject(ServerObject->NetRefHandle), nullptr);
+
+	// Deliver a packet
+	Server->UpdateAndSend({Client});
+
+	// Verify that object is now destroyed
+	UE_NET_ASSERT_EQ(Client->GetReplicationBridge()->GetReplicatedObject(OriginalObjectHandle), nullptr);
+	// Verify that new instance of object is not destroyed
+	UE_NET_ASSERT_NE(Client->GetReplicationBridge()->GetReplicatedObject(ServerObject->NetRefHandle), nullptr);
+
+	// Destroy object
+	Server->DestroyObject(ServerObject, EEndReplicationFlags::Destroy);
+
+	// Deliver a packet
+	Server->UpdateAndSend({Client});
+
+	// Verify that object now is destroyed
+	UE_NET_ASSERT_EQ(Client->GetReplicationBridge()->GetReplicatedObject(ServerObject->NetRefHandle), nullptr);
+}
+
+UE_NET_TEST_FIXTURE(FTestFlushBeforeDestroyFixture, TestDynamicStartStopStartAfterUpdateReplicatesAsNewObject)
+{
+	//LOG_SCOPE_VERBOSITY_OVERRIDE(LogIris, ELogVerbosity::Verbose);
+
+	FReplicationSystemTestClient* Client = CreateClient();
+	RegisterNetBlobHandlers(Client);
+
+	// Create and start to replicate object
+	UReplicatedTestObject* ServerObject = Server->CreateObject(0, 0);
+	FNetRefHandle OriginalServerObjectHandle = ServerObject->NetRefHandle;
+
+	// Stop replicating object using flush and without invalidating the server object.
+	Server->ReplicationBridge->EndReplication(ServerObject, EEndReplicationFlags::Destroy | EEndReplicationFlags::Flush);
+
+	// Send update, it should send the data.
+	Server->UpdateAndSend({Client});
+
+	// Verify that object is created which it should as we requested a flush.
+	UE_NET_ASSERT_NE(Client->GetReplicationBridge()->GetReplicatedObject(OriginalServerObjectHandle), nullptr);
+
+	// Now we start replicating the same object again
+	Server->ReplicationBridge->BeginReplication(ServerObject);
+
+	// Which should get a new handle as it is dynamic.
+	UE_NET_ASSERT_TRUE(ServerObject->NetRefHandle != OriginalServerObjectHandle);
+
+	// Send update, it should send the data.
+	Server->UpdateAndSend({Client});
+
+	// Verify that original instance is destroyed on client
+	UE_NET_ASSERT_EQ(Client->GetReplicationBridge()->GetReplicatedObject(OriginalServerObjectHandle), nullptr);
+
+	// Verify that new instance of object is created
+	UE_NET_ASSERT_NE(Client->GetReplicationBridge()->GetReplicatedObject(ServerObject->NetRefHandle), nullptr);
+}
+
+UE_NET_TEST_FIXTURE(FTestFlushBeforeDestroyFixture, TestDynamicStartStopStartAfterUpdateReplicatesAsNewObjectWithPacketLoss)
+{
+	//LOG_SCOPE_VERBOSITY_OVERRIDE(LogIris, ELogVerbosity::Verbose);
+
+	FReplicationSystemTestClient* Client = CreateClient();
+	RegisterNetBlobHandlers(Client);
+
+	// Create and start to replicate object
+	UReplicatedTestObject* ServerObject = Server->CreateObject(0, 0);
+	FNetRefHandle OriginalServerObjectHandle = ServerObject->NetRefHandle;
+
+	// Create reliable attachment to ServerObject to initiate proper flush, can be removed when UE-358165 is fixed.
+	EnqueueAttachmentToClients({Client}, ServerObject->NetRefHandle);
+
+	// Stop replicating object that should be flushed due to reliable-rpc and without invalidating the server object.
+	Server->ReplicationBridge->EndReplication(ServerObject, EEndReplicationFlags::Destroy | EEndReplicationFlags::Flush);
+
+	// Send update, but drop the packet
+	Server->UpdateAndSend({Client}, false);
+
+	// Now we start replicating the same object instance again, 
+	Server->ReplicationBridge->BeginReplication(ServerObject);
+
+	// Which should get a new handle as it is dynamic.
+	UE_NET_ASSERT_TRUE(ServerObject->NetRefHandle != OriginalServerObjectHandle);
+
+	Server->UpdateAndSend({Client});
+
+	// This can be surprising, but due to how dynamic objects work (that they get assigned a new NetRefHandle every time we start replication on the server) we will see both instances represented on the client
+
+	// Verify that object is created which it should as we requested a flush.
+	UE_NET_ASSERT_NE(Client->GetReplicationBridge()->GetReplicatedObject(OriginalServerObjectHandle), nullptr);
+
+	// Verify that new instance of object is created
+	UE_NET_ASSERT_NE(Client->GetReplicationBridge()->GetReplicatedObject(ServerObject->NetRefHandle), nullptr);
+
+	Server->UpdateAndSend({Client});
+
+	// Now we expect everything to reflect the server state.
+
+	// Verify that original instance is destroyed on client
+	UE_NET_ASSERT_EQ(Client->GetReplicationBridge()->GetReplicatedObject(OriginalServerObjectHandle), nullptr);
+
+	// Verify that new instance of object is created
+	UE_NET_ASSERT_NE(Client->GetReplicationBridge()->GetReplicatedObject(ServerObject->NetRefHandle), nullptr);
+}
+
+UE_NET_TEST_FIXTURE(FTestFlushBeforeDestroyFixture, TestDynamicSubObjectStartStopStartReplicatesAsNewObject)
+{
+	FReplicationSystemTestClient* Client = CreateClient();
+
+	// Create and start to replicate object
+	UReplicatedTestObject* ServerObject = Server->CreateObject(0, 0);
+	FNetRefHandle ObjectHandle = ServerObject->NetRefHandle;
+	UReplicatedTestObject* ServerSubObject = Server->CreateSubObject(ObjectHandle, 0, 0);
+	FNetRefHandle OriginalSubObjectHandle = ServerSubObject->NetRefHandle;
+
+	// Stop replicating SubObject but reuse instance indicating that it should be flushed i.e. that the final state should be replicated to all clients scoping the object during the next NetUpdate.
+	Server->ReplicationBridge->EndReplication(ServerSubObject, EEndReplicationFlags::Destroy | EEndReplicationFlags::Flush);
+
+	// Start replicating SubObject again canceling the stop replication request
+	Server->ReplicationBridge->BeginReplication(ServerSubObject);
+
+	// Expect new handle as it is a dynamic instance
+	UE_NET_ASSERT_TRUE(ServerSubObject->NetRefHandle != OriginalSubObjectHandle);
+
+	Server->UpdateAndSend({Client});
+
+	// Verify that objects are created
+	UE_NET_ASSERT_TRUE(Client->GetReplicationBridge()->GetReplicatedObject(ObjectHandle) != nullptr);
+	UE_NET_ASSERT_TRUE(Client->GetReplicationBridge()->GetReplicatedObject(OriginalSubObjectHandle) != nullptr);
+	UE_NET_ASSERT_TRUE(Client->GetReplicationBridge()->GetReplicatedObject(ServerSubObject->NetRefHandle) != nullptr);
+
+	Server->UpdateAndSend({Client});
+
+	// Verify that objects still exists
+	UE_NET_ASSERT_TRUE(Client->GetReplicationBridge()->GetReplicatedObject(ObjectHandle) != nullptr);
+	UE_NET_ASSERT_TRUE(Client->GetReplicationBridge()->GetReplicatedObject(ServerSubObject->NetRefHandle) != nullptr);
+
+	// But original subobject is destroyed
+	UE_NET_ASSERT_FALSE(Client->GetReplicationBridge()->GetReplicatedObject(OriginalSubObjectHandle) != nullptr);
 }
+
+UE_NET_TEST_FIXTURE(FTestFlushBeforeDestroyFixture, TestDynamicSubObjectStartStopStartSameFrameReplicatesAsNewObject)
+{
+	FReplicationSystemTestClient* Client = CreateClient();
+
+	// Create and start to replicate object
+	UReplicatedTestObject* ServerObject = Server->CreateObject(0, 0);
+	FNetRefHandle ObjectHandle = ServerObject->NetRefHandle;
+	UReplicatedTestObject* ServerSubObject = Server->CreateSubObject(ObjectHandle, 0, 0);
+	const FNetRefHandle OriginalSubObjectHandle = ServerSubObject->NetRefHandle;
+
+	// Stop replicating SubObject but keep server instance
+	Server->ReplicationBridge->EndReplication(ServerSubObject, EEndReplicationFlags::Destroy);
+
+	// Start replicating SubObject again, will be assigned new Handle as it is dynamic
+	Server->ReplicationBridge->BeginReplication(ServerSubObject);
+	const FNetRefHandle SubObjectHandle = ServerSubObject->NetRefHandle;
+
+	// Expect new handle
+	UE_NET_ASSERT_TRUE(ServerSubObject->NetRefHandle != OriginalSubObjectHandle);
+
+	Server->UpdateAndSend({Client});
+
+	// Verify that objects are created
+	UE_NET_ASSERT_TRUE(Client->GetReplicationBridge()->GetReplicatedObject(ObjectHandle) != nullptr);
+	UE_NET_ASSERT_TRUE(Client->GetReplicationBridge()->GetReplicatedObject(SubObjectHandle) != nullptr);
+	// OriginalSubObjectHandle was never relevant for scoping so it will not be created.
+	UE_NET_ASSERT_TRUE(Client->GetReplicationBridge()->GetReplicatedObject(OriginalSubObjectHandle) == nullptr);
+
+	Server->UpdateAndSend({Client});
+
+	// Verify that objects still exists
+	UE_NET_ASSERT_TRUE(Client->GetReplicationBridge()->GetReplicatedObject(ObjectHandle) != nullptr);
+	UE_NET_ASSERT_TRUE(Client->GetReplicationBridge()->GetReplicatedObject(SubObjectHandle) != nullptr);
+}
+
+UE_NET_TEST_FIXTURE(FTestFlushBeforeDestroyFixture, TestDynamicSubObjectStartFlushStartSameFrameReplicatesAsNewObject)
+{
+	FReplicationSystemTestClient* Client = CreateClient();
+
+	// Create and start to replicate object
+	UReplicatedTestObject* ServerObject = Server->CreateObject(0, 0);
+	FNetRefHandle ObjectHandle = ServerObject->NetRefHandle;
+	UReplicatedTestObject* ServerSubObject = Server->CreateSubObject(ObjectHandle, 0, 0);
+	FNetRefHandle OriginalSubObjectHandle = ServerSubObject->NetRefHandle;
+
+	// Stop replicating SubObject but reuse instance indicating that it should be flushed that is that the final state should be replicated to all clients scoping the object this frame
+	Server->ReplicationBridge->EndReplication(ServerSubObject, EEndReplicationFlags::Destroy | EEndReplicationFlags::Flush);
+
+	// Start replicating SubObject
+	Server->ReplicationBridge->BeginReplication(ServerSubObject);
+
+	Server->UpdateAndSend({Client});
+
+	// Verify that objects are created
+	UE_NET_ASSERT_TRUE(Client->GetReplicationBridge()->GetReplicatedObject(ObjectHandle) != nullptr);
+	UE_NET_ASSERT_TRUE(Client->GetReplicationBridge()->GetReplicatedObject(OriginalSubObjectHandle) != nullptr);
+	UE_NET_ASSERT_TRUE(Client->GetReplicationBridge()->GetReplicatedObject(ServerSubObject->NetRefHandle) != nullptr);
+
+	Server->UpdateAndSend({Client});
+
+	// Verify that original subobject is destroyed
+	UE_NET_ASSERT_TRUE(Client->GetReplicationBridge()->GetReplicatedObject(ObjectHandle) != nullptr);
+	UE_NET_ASSERT_FALSE(Client->GetReplicationBridge()->GetReplicatedObject(OriginalSubObjectHandle) != nullptr);
+	UE_NET_ASSERT_TRUE(Client->GetReplicationBridge()->GetReplicatedObject(ServerSubObject->NetRefHandle) != nullptr);
+}
+
+UE_NET_TEST_FIXTURE(FTestFlushBeforeDestroyFixture, TestStableObjectStartFlushStartAfterNextUpdate)
+{
+	//LOG_SCOPE_VERBOSITY_OVERRIDE(LogIris, ELogVerbosity::Verbose);
+
+	const TCHAR* StableObjectName = TEXT("StableObjectName");
+
+	UReplicationSystem* ReplicationSystem = Server->ReplicationSystem;
+
+	// Add clients
+	FReplicationSystemTestClient* Client = CreateClient();
+
+	// Create stably named object on both server and clients
+	UTestReplicatedIrisObject* ServerObject = Server->CreateStablyNamedObject<UTestReplicatedIrisObject>(StableObjectName);
+	UTestReplicatedIrisObject* ClientObject = Client->CreateStablyNamedObject<UTestReplicatedIrisObject>(StableObjectName);
+
+	UE_NET_ASSERT_TRUE(ServerObject->NetRefHandle.IsValid());
+	UE_NET_ASSERT_TRUE(ServerObject->NetRefHandle.IsStatic());
+
+	// Create and start to replicate object
+	FNetRefHandle OriginalServerObjectHandle = ServerObject->NetRefHandle;
+
+	ServerObject->IntA = 1;
+
+	// Stop replicating object using flush and without invalidating the server object.
+	// Flush means that we send the final state and attachments to the clients.
+	Server->ReplicationBridge->EndReplication(ServerObject, EEndReplicationFlags::Flush);
+
+	// After the update, the server will mark the flushed NetObject/NetRefHandle as destroyed.
+	Server->UpdateAndSend({Client});
+
+	// Verify that object is created which it should as we requested a flush.
+	UE_NET_ASSERT_TRUE(Cast<UTestReplicatedIrisObject>(Client->GetReplicationBridge()->GetReplicatedObject(OriginalServerObjectHandle)) == ClientObject);
+	UE_NET_ASSERT_EQ(ClientObject->IntA, 1);
+
+	// Now we start replicating the same object again, which will be assigned a new internal index.
+	Server->ReplicationBridge->BeginReplication(ServerObject);
+
+	// As this is a static handle we expect the same handle to be used
+	UE_NET_ASSERT_TRUE(ServerObject->NetRefHandle == OriginalServerObjectHandle);
+
+	ServerObject->IntA = 2;
+
+	// This update will send the actual destroy of the original (flushed) instance.
+	Server->UpdateAndSend({Client});
+
+	// Verify that original instance is no longer considered replicated
+	UE_NET_ASSERT_EQ(Client->GetReplicationBridge()->GetReplicatedObject(OriginalServerObjectHandle), nullptr);
+
+	// Now that know that the client has destroyed the NetObject for the flushed instance we expect the enw instance to replicate.
+	Server->UpdateAndSend({Client});
+
+	// Verify that NetRefHandle is now considered replicated again
+	UE_NET_ASSERT_TRUE(Cast<UTestReplicatedIrisObject>(Client->GetReplicationBridge()->GetReplicatedObject(ServerObject->NetRefHandle)) == ClientObject);
+	UE_NET_ASSERT_EQ(ClientObject->IntA, 2);
+}
+
+UE_NET_TEST_FIXTURE(FTestFlushBeforeDestroyFixture, TestStableObjectCreatedAndDestroyedAndRestartedAfterNextUpdateReplicatesIfFlushedWithDroppedData)
+{
+	//LOG_SCOPE_VERBOSITY_OVERRIDE(LogIris, ELogVerbosity::Verbose);
+
+	const TCHAR* StableObjectName = TEXT("StableObjectName");
+
+	UReplicationSystem* ReplicationSystem = Server->ReplicationSystem;
+
+	// Add clients
+	FReplicationSystemTestClient* Client = CreateClient();
+	RegisterNetBlobHandlers(Client);
+
+	// Create stably named object on both server and clients
+	UTestReplicatedIrisObject* ServerObject = Server->CreateStablyNamedObject<UTestReplicatedIrisObject>(StableObjectName);
+	UTestReplicatedIrisObject* ClientObject = Client->CreateStablyNamedObject<UTestReplicatedIrisObject>(StableObjectName);
+
+	UE_NET_ASSERT_TRUE(ServerObject->NetRefHandle.IsValid());
+	UE_NET_ASSERT_TRUE(ServerObject->NetRefHandle.IsStatic());
+
+	// Create and start to replicate object
+	FNetRefHandle OriginalServerObjectHandle = ServerObject->NetRefHandle;
+
+	ServerObject->IntA = 1;
+
+	// Create reliable attachment to ServerObject to initiate proper flush, can be removed when UE-358165 is fixed.
+	EnqueueAttachmentToClients({Client}, ServerObject->NetRefHandle);
+	
+	// Stop replicating object using flush and without invalidating the server object.
+	// Flush means that we send the final state and attachments to the clients.
+	Server->ReplicationBridge->EndReplication(ServerObject, EEndReplicationFlags::Flush);
+
+	// After the update, the server will mark the flushed NetObject/NetRefHandle as destroyed.
+	// Note: that we drop the data
+	Server->UpdateAndSend({Client}, false);
+
+	// Now we start replicating the same object again, which will be assigned a new internal index.
+	Server->ReplicationBridge->BeginReplication(ServerObject);
+
+	ServerObject->IntA = 2;
+
+	// As this is a static handle we expect the same handle to be used
+	UE_NET_ASSERT_TRUE(ServerObject->NetRefHandle == OriginalServerObjectHandle);
+
+	// After the update, the server will mark the flushed NetObject/NetRefHandle as destroyed.
+	// Note: that we drop the data
+	Server->UpdateAndSend({Client});
+
+	// Verify that object is created which it should as we requested a flush, even after dropping the packet containing the original create request
+	UE_NET_ASSERT_TRUE(Cast<UTestReplicatedIrisObject>(Client->GetReplicationBridge()->GetReplicatedObject(OriginalServerObjectHandle)) == ClientObject);
+	UE_NET_ASSERT_EQ(ClientObject->IntA, 1);
+
+	// This update will send the actual destroy of the original (flushed) instance.
+	Server->UpdateAndSend({Client});
+
+	// Verify that original instance is no longer considered replicated
+	UE_NET_ASSERT_EQ(Cast<UTestReplicatedIrisObject>(Client->GetReplicationBridge()->GetReplicatedObject(OriginalServerObjectHandle)), nullptr);
+
+	// Now that know that the client has destroyed the NetObject for the flushed instance we expect the new instance to replicate.
+	Server->UpdateAndSend({Client});
+
+	// Verify that NetRefHandle is now considered replicated again
+	UE_NET_ASSERT_TRUE(Cast<UTestReplicatedIrisObject>(Client->GetReplicationBridge()->GetReplicatedObject(ServerObject->NetRefHandle)) == ClientObject);
+	UE_NET_ASSERT_EQ(ClientObject->IntA, 2);
+}
+
+UE_NET_TEST_FIXTURE(FTestFlushBeforeDestroyFixture, TestStateFlushedDoesNotPickupStateChangesAfterFlush)
+{
+	FReplicationSystemTestClient* Client = CreateClient();
+	RegisterNetBlobHandlers(Client);
+
+	UTestReplicatedIrisObject* ServerObject = Server->CreateObject(0, 0);
+	FNetRefHandle ObjectHandle = ServerObject->NetRefHandle;
+	
+	Server->UpdateAndSend({Client});
+	
+	// Verify that object is created
+	UE_NET_ASSERT_NE(Client->GetReplicationBridge()->GetReplicatedObject(ObjectHandle), nullptr);
+
+	// Modify state
+	ServerObject->IntA = 1;
+
+	// Destroy object with flush flag which should flush the state before destroying the object
+	Server->ReplicationBridge->EndReplication(ServerObject, EEndReplicationFlags::Destroy | EEndReplicationFlags::Flush);
+
+	// Modify state, this should not be sent
+	ServerObject->IntA = 2;
+
+	Server->UpdateAndSend({Client});
+
+	UTestReplicatedIrisObject* ClientObject = Cast<UTestReplicatedIrisObject>(Client->GetReplicationBridge()->GetReplicatedObject(ObjectHandle));
+
+	// Verify that object is created
+	UE_NET_ASSERT_NE(ClientObject, nullptr);
+
+	// Verify that we got the expected state
+	UE_NET_ASSERT_EQ(ClientObject->IntA, 1);
+
+	// Deliver a packet. Should destroy the object on the client unless that was done
+	Server->UpdateAndSend({Client});
+
+	// Verify that object is destroyed
+	UE_NET_ASSERT_EQ(Client->GetReplicationBridge()->GetReplicatedObject(ObjectHandle), nullptr);
+}
+
+UE_NET_TEST_FIXTURE(FTestFlushBeforeDestroyFixture, TestStateFlushedDoesNotAccecptNewSubObjectAfterFlush)
+{
+	FReplicationSystemTestClient* Client = CreateClient();
+	RegisterNetBlobHandlers(Client);
+
+	UTestReplicatedIrisObject* ServerObject = Server->CreateObject(0, 0);
+	FNetRefHandle ObjectHandle = ServerObject->NetRefHandle;
+
+	// Modify state
+	ServerObject->IntA = 1;
+
+	// Destroy object with flush flag which should flush the state before destroying the object
+	Server->ReplicationBridge->EndReplication(ServerObject, EEndReplicationFlags::Destroy | EEndReplicationFlags::Flush);
+
+	// Modify state
+	ServerObject->IntA = 2;
+
+	// Add SubObject
+	{
+		// This should ensure so we add a scope to suppress it
+		FTestEnsureScope EnsureScope;
+
+		UReplicatedTestObject* ServerSubObject = Server->CreateSubObject(ServerObject->NetRefHandle, 0, 0);
+		FNetRefHandle SubObjectHandle = ServerSubObject->NetRefHandle;
+
+		UE_NET_ASSERT_FALSE(SubObjectHandle.IsValid());
+	}
+		
+	Server->UpdateAndSend({Client});
+
+	// Verify that object is created and that we got the expected state
+	{
+		UTestReplicatedIrisObject* ClientObject = Cast<UTestReplicatedIrisObject>(Client->GetReplicationBridge()->GetReplicatedObject(ObjectHandle));
+		UE_NET_ASSERT_NE(ClientObject, nullptr);
+		// Verify that we got the expected state
+		UE_NET_ASSERT_EQ(ClientObject->IntA, 1);
+	}
+
+	Server->UpdateAndSend({Client});
+
+	// Verify that object is destroyed
+	UE_NET_ASSERT_EQ(Cast<UTestReplicatedIrisObject>(Client->GetReplicationBridge()->GetReplicatedObject(ObjectHandle)), nullptr);
+}
+
+} // End of namespace
diff --git a/Engine/Plugins/Runtime/ReplicationSystemTestPlugin/Source/Private/Tests/ReplicationSystem/TestReplicationSystem.cpp b/Engine/Plugins/Runtime/ReplicationSystemTestPlugin/Source/Private/Tests/ReplicationSystem/TestReplicationSystem.cpp
--- a/Engine/Plugins/Runtime/ReplicationSystemTestPlugin/Source/Private/Tests/ReplicationSystem/TestReplicationSystem.cpp
+++ b/Engine/Plugins/Runtime/ReplicationSystemTestPlugin/Source/Private/Tests/ReplicationSystem/TestReplicationSystem.cpp
@@ -1234,7 +1234,45 @@ UE_NET_TEST_FIXTURE(FReplicationSystemServerClientTestFixture, TestTearOffImmedi
 	// TearOff the object before first replication
 	Server->ReplicationBridge->EndReplication(ServerObject, EEndReplicationFlags::TearOff);
 
+	// Send and do not deliver packet
+	Server->UpdateAndSend({Client}, false);
+
 	// Send and deliver packet
+	Server->UpdateAndSend({Client});
+
+	// Client should have created a object
+	UE_NET_ASSERT_EQ(NumObjectsCreatedOnClientBeforeReplication + 1, Client->CreatedObjects.Num());
+
+	UE_NET_ASSERT_TRUE(Cast<UTestReplicatedIrisObject>(Client->GetReplicationBridge()->GetReplicatedObject(ServerObject->NetRefHandle)) == nullptr);
+
+	// We should be able to get the object from the created objects array to validate the state
+	UTestReplicatedIrisObject* ClientObjectThatWillBeTornOff = Cast<UTestReplicatedIrisObject>(Client->CreatedObjects[NumObjectsCreatedOnClientBeforeReplication].Get());
+
+	// Verify that we replicated the expected state
+	UE_NET_ASSERT_EQ(ServerObject->IntA, ClientObjectThatWillBeTornOff->IntA);
+}
+
+// Test TearOff for new object and resend if not fitting in packet, this requires creation info to be cached and that we retain object in scope
+UE_NET_TEST_FIXTURE(FReplicationSystemServerClientTestFixture, TestTearOffImmediateOnNewlyCreatedObjectIfNotSentOnFirstUpdate)
+{
+	UReplicationSystem* ReplicationSystem = Server->ReplicationSystem;
+
+	// Add a client
+	FReplicationSystemTestClient* Client = CreateClient();
+
+	// We should not have any created objects
+	const int32 NumObjectsCreatedOnClientBeforeReplication = Client->CreatedObjects.Num();
+
+	// Spawn object on server
+	UTestReplicatedIrisObject* ServerObject = Server->CreateObject(0,0);
+
+	// Set state
+	ServerObject->IntA = 1;
+
+	// TearOff the object before first replication
+	Server->ReplicationBridge->EndReplication(ServerObject, EEndReplicationFlags::TearOff);
+
+	// Send and do not deliver packet
 	Server->UpdateAndSend({Client}, false);
 
 	// Send and deliver packet
@@ -1657,16 +1695,19 @@ UE_NET_TEST_FIXTURE(FReplicationSystemServerClientTestFixture, TestCreatingSubob
 		{
 			LOG_SCOPE_VERBOSITY_OVERRIDE(LogIris, ELogVerbosity::Error);
 			ServerSubObject1 = Server->CreateSubObject(ServerObject->NetRefHandle);
-			UE_NET_ASSERT_EQ(EnsureScope.GetCount(), 0);
-		}
 
-		// The subobject should not start replicating.
-		UE_NET_ASSERT_FALSE(ServerSubObject1->NetRefHandle.IsValid());
+			// We expect an ensure here, as the object should not longer be considered replicated.
+			// Cannot really rely on ensure counts as other tests might affect it if the same ensure has already been triggered...
+			//UE_NET_ASSERT_EQ(EnsureScope.GetCount(), 1);
 
-		// Send and deliver packet
-		Server->UpdateAndSend({ Client });
-		UE_NET_ASSERT_EQ(EnsureScope.GetCount(), 0);
+			// The subobject should not start replicating.
+			UE_NET_ASSERT_FALSE(ServerSubObject1->NetRefHandle.IsValid());
+		}
 	}
+
+	// Send and deliver packet
+	Server->UpdateAndSend({ Client });
+
 }
 
 // Test dropped creation of subobject dirties owner
@@ -2728,7 +2769,83 @@ UE_NET_TEST_FIXTURE(FReplicationSystemServerClientTestFixture, TestDynamicObject
 	UE_NET_ASSERT_NE(OriginalServerHandle, ServerObject->NetRefHandle);	
 }
 
-UE_NET_TEST_FIXTURE(FReplicationSystemServerClientTestFixture, TestDynamicObjectAsyncStopReplicationAndFlushWitTwoClients)
+UE_NET_TEST_FIXTURE(FReplicationSystemServerClientTestFixture, TestDynamicObjectAsyncStopReplicationForSubObject)
+{
+	LOG_SCOPE_VERBOSITY_OVERRIDE(LogIris, ELogVerbosity::Verbose);
+
+	UReplicationSystem* ReplicationSystem = Server->ReplicationSystem;
+
+	// Add a client
+	FReplicationSystemTestClient* Client = CreateClient();
+
+	// Create dynamic object
+	UTestReplicatedIrisObject* ServerObject = Server->CreateObject<UTestReplicatedIrisObject>();
+
+	// Create subobject
+	UReplicatedTestObject* ServerSubObject = Server->CreateSubObject(ServerObject->NetRefHandle, 0, 0);
+	const FNetRefHandle ServerSubObjectRefHandle = ServerSubObject->NetRefHandle;
+
+	UE_NET_ASSERT_TRUE(ServerObject->NetRefHandle.IsValid());
+	UE_NET_ASSERT_TRUE(ServerObject->NetRefHandle.IsDynamic());
+
+	const FNetRefHandle OriginalServerSubObjectHandle = ServerSubObject->NetRefHandle;
+
+	// Send and deliver packet
+	Server->UpdateAndSend({Client});
+
+	// We expect object to exist on both server and client
+	UE_NET_ASSERT_TRUE(Cast<UTestReplicatedIrisObject>(Client->GetReplicationBridge()->GetReplicatedObject(ServerObject->NetRefHandle)) != nullptr);
+	UE_NET_ASSERT_TRUE(Cast<UTestReplicatedIrisObject>(Server->GetReplicationBridge()->GetReplicatedObject(ServerObject->NetRefHandle)) != nullptr);
+	UE_NET_ASSERT_TRUE(Cast<UReplicatedTestObject>(Client->GetReplicationBridge()->GetReplicatedObject(ServerSubObject->NetRefHandle)) != nullptr);
+	UE_NET_ASSERT_TRUE(Cast<UReplicatedTestObject>(Server->GetReplicationBridge()->GetReplicatedObject(ServerSubObject->NetRefHandle)) != nullptr);
+
+	// Stop replication of subobject using async path
+	Server->ReplicationBridge->EndReplication(ServerSubObject, EEndReplicationFlags::AllowAsyncStopReplication | EEndReplicationFlags::Destroy);
+
+	// Should be pending stop replication
+	UE_NET_ASSERT_TRUE(Server->ReplicationBridge->IsPendingStopReplication(ServerSubObject->NetRefHandle));
+	UE_NET_ASSERT_TRUE(Server->ReplicationBridge->IsReplicatedHandle(ServerSubObject->NetRefHandle));
+	
+	// This should cancel the async stop replication and keep replicating the object using the same handle
+	Server->ReplicationBridge->BeginReplication(ServerObject->NetRefHandle, ServerSubObject);
+
+	// Verify that we did not get a new handle
+	UE_NET_ASSERT_EQ(OriginalServerSubObjectHandle, ServerSubObject->NetRefHandle);
+	UE_NET_ASSERT_TRUE(Server->ReplicationBridge->IsReplicatedHandle(ServerSubObject->NetRefHandle));
+
+	// Should not be pending stop replication
+	UE_NET_ASSERT_FALSE(Server->ReplicationBridge->IsPendingStopReplication(ServerSubObject->NetRefHandle));
+
+	// Send and deliver packet
+	Server->UpdateAndSend({Client});
+		
+	// We expect object to exist on both server and client
+	UE_NET_ASSERT_TRUE(Cast<UTestReplicatedIrisObject>(Client->GetReplicationBridge()->GetReplicatedObject(ServerObject->NetRefHandle)) != nullptr);
+	UE_NET_ASSERT_TRUE(Cast<UTestReplicatedIrisObject>(Server->GetReplicationBridge()->GetReplicatedObject(ServerObject->NetRefHandle)) != nullptr);
+	UE_NET_ASSERT_TRUE(Cast<UReplicatedTestObject>(Client->GetReplicationBridge()->GetReplicatedObject(ServerSubObject->NetRefHandle)) != nullptr);
+	UE_NET_ASSERT_TRUE(Cast<UReplicatedTestObject>(Server->GetReplicationBridge()->GetReplicatedObject(ServerSubObject->NetRefHandle)) != nullptr);
+
+	// Stop replication using async path
+	Server->ReplicationBridge->EndReplication(ServerSubObject, EEndReplicationFlags::AllowAsyncStopReplication | EEndReplicationFlags::Destroy);
+
+	// Send and deliver packet
+	Server->UpdateAndSend({Client});
+
+	// We expect Object to still exist on client and server
+	UE_NET_ASSERT_TRUE(Cast<UTestReplicatedIrisObject>(Client->GetReplicationBridge()->GetReplicatedObject(ServerObject->NetRefHandle)) != nullptr);
+	UE_NET_ASSERT_TRUE(Cast<UTestReplicatedIrisObject>(Server->GetReplicationBridge()->GetReplicatedObject(ServerObject->NetRefHandle)) != nullptr);
+
+	// We expect SubObject to no longer exist on client and server
+	UE_NET_ASSERT_TRUE(Cast<UReplicatedTestObject>(Client->GetReplicationBridge()->GetReplicatedObject(ServerSubObject->NetRefHandle)) == nullptr);
+	UE_NET_ASSERT_TRUE(Cast<UReplicatedTestObject>(Server->GetReplicationBridge()->GetReplicatedObject(ServerSubObject->NetRefHandle)) == nullptr);
+
+	// This should start replicating the object again using a new handle, as the object have completed the async stop replication.
+	Server->ReplicationBridge->BeginReplication(ServerObject->NetRefHandle, ServerSubObject);
+
+	UE_NET_ASSERT_NE(OriginalServerSubObjectHandle, ServerSubObject->NetRefHandle);	
+}
+
+UE_NET_TEST_FIXTURE(FReplicationSystemServerClientTestFixture, TestDynamicObjectAsyncStopReplicationWitTwoClients)
 {
 	UReplicationSystem* ReplicationSystem = Server->ReplicationSystem;
 
@@ -2747,7 +2864,7 @@ UE_NET_TEST_FIXTURE(FReplicationSystemServerClientTestFixture, TestDynamicObject
 	// Send and deliver packet to clients
 	Server->UpdateAndSend(Clients);
 
-	// The object should now exist on all clients
+	// The object should now exist on all clients and server
 	UE_NET_ASSERT_TRUE(Cast<UTestReplicatedIrisObject>(Clients[0]->GetReplicationBridge()->GetReplicatedObject(ServerObject->NetRefHandle)) != nullptr);
 	UE_NET_ASSERT_TRUE(Cast<UTestReplicatedIrisObject>(Clients[1]->GetReplicationBridge()->GetReplicatedObject(ServerObject->NetRefHandle)) != nullptr);
 	UE_NET_ASSERT_TRUE(Cast<UTestReplicatedIrisObject>(Server->GetReplicationBridge()->GetReplicatedObject(ServerObject->NetRefHandle)) != nullptr);
@@ -2779,7 +2896,7 @@ UE_NET_TEST_FIXTURE(FReplicationSystemServerClientTestFixture, TestDynamicObject
 	// Send and deliver packet to clients
 	Server->UpdateAndSend(Clients);
 
-	// The object should now exist on all clients
+	// The object should now exist on all clients again
 	UE_NET_ASSERT_TRUE(Cast<UTestReplicatedIrisObject>(Clients[0]->GetReplicationBridge()->GetReplicatedObject(ServerObject->NetRefHandle)) != nullptr);
 	UE_NET_ASSERT_TRUE(Cast<UTestReplicatedIrisObject>(Clients[1]->GetReplicationBridge()->GetReplicatedObject(ServerObject->NetRefHandle)) != nullptr);
 	UE_NET_ASSERT_TRUE(Cast<UTestReplicatedIrisObject>(Server->GetReplicationBridge()->GetReplicatedObject(ServerObject->NetRefHandle)) != nullptr);
@@ -2818,6 +2935,236 @@ UE_NET_TEST_FIXTURE(FReplicationSystemServerClientTestFixture, TestDynamicObject
 	UE_NET_ASSERT_TRUE(Cast<UTestReplicatedIrisObject>(Clients[1]->GetReplicationBridge()->GetReplicatedObject(ServerObject->NetRefHandle)) != nullptr);
 }
 
+UE_NET_TEST_FIXTURE(FReplicationSystemServerClientTestFixture, TestStablyNamedObjectStartStopStartReplicationWithTwoClients)
+{
+	//LOG_SCOPE_VERBOSITY_OVERRIDE(LogIris, ELogVerbosity::Verbose);
+
+	const TCHAR* StableObjectName = TEXT("StableObjectName");
+
+	UReplicationSystem* ReplicationSystem = Server->ReplicationSystem;
+
+	// Add clients
+	CreateClient();
+	CreateClient();
+
+	// Create stably named object on both server and clients
+	UTestReplicatedIrisObject* ServerObject = Server->CreateStablyNamedObject<UTestReplicatedIrisObject>(StableObjectName);
+	TArray<UTestReplicatedIrisObject*> ClientObjects;
+	ClientObjects.Add(Clients[0]->CreateStablyNamedObject<UTestReplicatedIrisObject>(StableObjectName));
+	ClientObjects.Add(Clients[1]->CreateStablyNamedObject<UTestReplicatedIrisObject>(StableObjectName));
+
+	UE_NET_ASSERT_TRUE(ServerObject->NetRefHandle.IsValid());
+	UE_NET_ASSERT_TRUE(ServerObject->NetRefHandle.IsStatic());
+
+	// Set some data so that we can identify the "incarnation"
+	ServerObject->IntA = 1;
+
+	// Send and deliver packet to first client but not to the other
+	Server->UpdateAndSend({Clients[0]});
+
+	// We expect object to exist on server and clients[0]
+	UE_NET_ASSERT_TRUE(Cast<UTestReplicatedIrisObject>(Clients[0]->GetReplicationBridge()->GetReplicatedObject(ServerObject->NetRefHandle)) == ClientObjects[0]);
+	UE_NET_ASSERT_FALSE(Cast<UTestReplicatedIrisObject>(Clients[0]->GetReplicationBridge()->GetReplicatedObject(ServerObject->NetRefHandle)) == ClientObjects[1]);
+	UE_NET_ASSERT_TRUE(Cast<UTestReplicatedIrisObject>(Server->GetReplicationBridge()->GetReplicatedObject(ServerObject->NetRefHandle)) != nullptr);
+
+	// Stop replication, not asking for a destroy as we do not want the client to destroy the instances.
+	Server->ReplicationBridge->EndReplication(ServerObject, EEndReplicationFlags::None);
+
+	// Should no longer be considered replicated
+	UE_NET_ASSERT_FALSE(Server->ReplicationBridge->IsReplicatedHandle(ServerObject->NetRefHandle));
+	
+	// Start replication again, but since we are starting to replicate an object already replicating to some clients there will have to ask the clients to conform the destroy before we can replicate the new instance.
+	Server->ReplicationBridge->BeginReplication(ServerObject);
+
+	// Set some data so that we can identify the "incarnation"
+	ServerObject->IntA = 2;
+
+	// Object should exist on server...but in its next incarnation, so replication will restart
+	UE_NET_ASSERT_TRUE(Cast<UTestReplicatedIrisObject>(Server->GetReplicationBridge()->GetReplicatedObject(ServerObject->NetRefHandle)) != nullptr);
+
+	// Send and deliver packet to both clients
+	Server->UpdateAndSend(Clients);
+
+	// As we already was replicating the object we need to wait for the clients to confirm this before we replicate a new instance of the same object,
+	// In this case as we had replicated the old instance to client[0] but not to client[1] so client[0] will have the old instance be destroyed while the client[1] will receive the new instnace
+	UE_NET_ASSERT_FALSE(Cast<UTestReplicatedIrisObject>(Clients[0]->GetReplicationBridge()->GetReplicatedObject(ServerObject->NetRefHandle)) == ClientObjects[0]);
+	UE_NET_ASSERT_TRUE(Cast<UTestReplicatedIrisObject>(Clients[1]->GetReplicationBridge()->GetReplicatedObject(ServerObject->NetRefHandle)) == ClientObjects[1]);
+
+	// Client[0] should have received the data from the first instance
+	UE_NET_ASSERT_EQ(ClientObjects[0]->IntA, 1);
+	// Client[1] should have received the data from got the latest instance
+	UE_NET_ASSERT_EQ(ClientObjects[1]->IntA, 2);
+
+	// Set some data so that we can identify the "incarnation"
+	ServerObject->IntA = 3;
+	
+	// Send and deliver packet to both clients
+	Server->UpdateAndSend(Clients);
+
+	// New instance should now have been replicated to both clients
+	UE_NET_ASSERT_TRUE(Cast<UTestReplicatedIrisObject>(Clients[0]->GetReplicationBridge()->GetReplicatedObject(ServerObject->NetRefHandle)) == ClientObjects[0]);
+	UE_NET_ASSERT_TRUE(Cast<UTestReplicatedIrisObject>(Clients[1]->GetReplicationBridge()->GetReplicatedObject(ServerObject->NetRefHandle)) == ClientObjects[1]);
+
+	// Both clients are now expected to have received the latest state.
+	UE_NET_ASSERT_EQ(ClientObjects[0]->IntA, 3);
+	UE_NET_ASSERT_EQ(ClientObjects[1]->IntA, 3);
+ 
+	Server->ReplicationBridge->EndReplication(ServerObject, EEndReplicationFlags::None);
+
+	// Update and send to all 
+	Server->UpdateAndSend(Clients);
+
+	// Object should now be destroyed on all clients and server
+	UE_NET_ASSERT_FALSE(Cast<UTestReplicatedIrisObject>(Clients[0]->GetReplicationBridge()->GetReplicatedObject(ServerObject->NetRefHandle)) == ClientObjects[0]);
+	UE_NET_ASSERT_FALSE(Cast<UTestReplicatedIrisObject>(Clients[1]->GetReplicationBridge()->GetReplicatedObject(ServerObject->NetRefHandle)) == ClientObjects[1]);
+	UE_NET_ASSERT_FALSE(Cast<UTestReplicatedIrisObject>(Server->GetReplicationBridge()->GetReplicatedObject(ServerObject->NetRefHandle)) != nullptr);
+}
+
+/** Test combinations of start / stop / start for stable object using flush. */
+UE_NET_TEST_FIXTURE(FReplicationSystemServerClientTestFixture, TestStableStartFlushStartStopStartReplication)
+{
+	LOG_SCOPE_VERBOSITY_OVERRIDE(LogIris, ELogVerbosity::Verbose);
+
+	const TCHAR* StableObjectName = TEXT("StableObjectName");
+
+	UReplicationSystem* ReplicationSystem = Server->ReplicationSystem;
+
+	// Add client
+	CreateClient();
+
+	// Create stably named object on both server and clients
+	UTestReplicatedIrisObject* ServerObject = Server->CreateStablyNamedObject<UTestReplicatedIrisObject>(StableObjectName);
+	TArray<UTestReplicatedIrisObject*> ClientObjects;
+	ClientObjects.Add(Clients[0]->CreateStablyNamedObject<UTestReplicatedIrisObject>(StableObjectName));
+
+	// Test multiple start stop during same frame
+	Server->ReplicationBridge->BeginReplication(ServerObject);
+
+	Server->UpdateAndSend(Clients);
+
+	ServerObject->IntA = 33;
+
+	// Stop replicating using flush
+	Server->ReplicationBridge->EndReplication(ServerObject, EEndReplicationFlags::Flush);
+
+	// Start replicating again
+	Server->ReplicationBridge->BeginReplication(ServerObject);
+
+	ServerObject->IntA = 44;
+
+	// Stop replication, without flush, we do not expect to see this state on the client
+	Server->ReplicationBridge->EndReplication(ServerObject);
+
+	// Start replicating again
+	Server->ReplicationBridge->BeginReplication(ServerObject);
+
+	ServerObject->IntA = 55;
+
+	// Stop replicating using flush
+	Server->ReplicationBridge->EndReplication(ServerObject, EEndReplicationFlags::Flush);
+
+	// Start replication again. As we have issued multiple start/stops on the same stable object requesting flush it will take a few updates for the final incarnation to be replicated.
+	Server->ReplicationBridge->BeginReplication(ServerObject);
+	ServerObject->IntA = 66;
+
+	Server->UpdateAndSend(Clients);
+
+	// We expect the first flushed incarnation to have been created
+	UE_NET_ASSERT_TRUE(Cast<UTestReplicatedIrisObject>(Clients[0]->GetReplicationBridge()->GetReplicatedObject(ServerObject->NetRefHandle)) == ClientObjects[0]);
+	UE_NET_ASSERT_EQ(ClientObjects[0]->IntA, 33);
+
+	Server->UpdateAndSend(Clients);
+
+	// Should be destroyed
+	UE_NET_ASSERT_TRUE(Cast<UTestReplicatedIrisObject>(Clients[0]->GetReplicationBridge()->GetReplicatedObject(ServerObject->NetRefHandle)) == nullptr);
+
+	Server->UpdateAndSend(Clients);
+
+	// We expect the next flushed incarnation to have been created
+	UE_NET_ASSERT_TRUE(Cast<UTestReplicatedIrisObject>(Clients[0]->GetReplicationBridge()->GetReplicatedObject(ServerObject->NetRefHandle)) == ClientObjects[0]);
+	UE_NET_ASSERT_EQ(ClientObjects[0]->IntA, 55);
+
+	Server->UpdateAndSend(Clients);
+
+	// Should be destroyed
+	UE_NET_ASSERT_TRUE(Cast<UTestReplicatedIrisObject>(Clients[0]->GetReplicationBridge()->GetReplicatedObject(ServerObject->NetRefHandle)) == nullptr);
+
+	Server->UpdateAndSend(Clients);
+
+	// We expect the final flushed incarnation to have been created
+	UE_NET_ASSERT_TRUE(Cast<UTestReplicatedIrisObject>(Clients[0]->GetReplicationBridge()->GetReplicatedObject(ServerObject->NetRefHandle)) == ClientObjects[0]);
+	UE_NET_ASSERT_EQ(ClientObjects[0]->IntA, 66);
+}
+
+UE_NET_TEST_FIXTURE(FReplicationSystemServerClientTestFixture, TestStablyNamedObjectMultiStartStopStartReplicationWithTwoClients)
+{
+	LOG_SCOPE_VERBOSITY_OVERRIDE(LogIris, ELogVerbosity::Verbose);
+
+	const TCHAR* StableObjectName = TEXT("StableObjectName");
+
+	UReplicationSystem* ReplicationSystem = Server->ReplicationSystem;
+
+	// Add clients
+	CreateClient();
+
+	// Create stably named object on both server and client
+	UTestReplicatedIrisObject* ServerObject = Server->CreateStablyNamedObject<UTestReplicatedIrisObject>(StableObjectName);
+	TArray<UTestReplicatedIrisObject*> ClientObjects;
+	ClientObjects.Add(Clients[0]->CreateStablyNamedObject<UTestReplicatedIrisObject>(StableObjectName));
+
+	// Test multiple start stop during same frame
+	Server->ReplicationBridge->BeginReplication(ServerObject);
+
+	// Update and send to all 
+	Server->UpdateAndSend(Clients);
+
+	ServerObject->IntA = 33;
+	// As we use flush, state will be captured immediately, and connections already scoping object will initiate flush.
+	Server->ReplicationBridge->EndReplication(ServerObject, EEndReplicationFlags::Flush);
+
+	// As we used flush, begin/start replication using the same internal index.
+	Server->ReplicationBridge->BeginReplication(ServerObject);
+
+	ServerObject->IntA = 55;
+
+	// Stop replication: Note: That this will fully detach instance from NetObject
+	Server->ReplicationBridge->EndReplication(ServerObject);
+
+	// Start replication again. As we are still replicating the first incarnation, there will be some latency before the last state is replicated as a new NetObject.
+	Server->ReplicationBridge->BeginReplication(ServerObject);
+	ServerObject->IntA = 66;
+
+	// Late join client 1, as it joined before the NetUpdate and we used flush on the original incarnation should be added to its scope as well.
+	CreateClient();
+	ClientObjects.Add(Clients[1]->CreateStablyNamedObject<UTestReplicatedIrisObject>(StableObjectName));
+
+	Server->UpdateAndSend(Clients);
+	
+	// Should exist on both clients, with he same state
+	UE_NET_ASSERT_TRUE(Cast<UTestReplicatedIrisObject>(Clients[0]->GetReplicationBridge()->GetReplicatedObject(ServerObject->NetRefHandle)) == ClientObjects[0]);
+	UE_NET_ASSERT_TRUE(Cast<UTestReplicatedIrisObject>(Clients[1]->GetReplicationBridge()->GetReplicatedObject(ServerObject->NetRefHandle)) == ClientObjects[1]);
+
+	// We expect the flushed state for both clients
+	UE_NET_ASSERT_EQ(ClientObjects[0]->IntA, 33);
+	UE_NET_ASSERT_EQ(ClientObjects[1]->IntA, 33);
+
+	Server->UpdateAndSend(Clients);
+
+	// Now we expect first incarnation to have been destroyed
+	UE_NET_ASSERT_TRUE(Cast<UTestReplicatedIrisObject>(Clients[0]->GetReplicationBridge()->GetReplicatedObject(ServerObject->NetRefHandle)) == nullptr);
+	UE_NET_ASSERT_TRUE(Cast<UTestReplicatedIrisObject>(Clients[1]->GetReplicationBridge()->GetReplicatedObject(ServerObject->NetRefHandle)) == nullptr);
+
+	Server->UpdateAndSend(Clients);
+
+	// Should now exist on both clients
+	UE_NET_ASSERT_TRUE(Cast<UTestReplicatedIrisObject>(Clients[0]->GetReplicationBridge()->GetReplicatedObject(ServerObject->NetRefHandle)) == ClientObjects[0]);
+	UE_NET_ASSERT_TRUE(Cast<UTestReplicatedIrisObject>(Clients[1]->GetReplicationBridge()->GetReplicatedObject(ServerObject->NetRefHandle)) == ClientObjects[1]);
+
+	// We expect the latest state should be delivered
+	UE_NET_ASSERT_EQ(ClientObjects[0]->IntA, 66);
+	UE_NET_ASSERT_EQ(ClientObjects[1]->IntA, 66);
+}
+
 } // end namespace UE::Net::Private
 
 
diff --git a/Engine/Source/Runtime/Net/Iris/Private/Iris/ReplicationSystem/NetRefHandleManager.cpp b/Engine/Source/Runtime/Net/Iris/Private/Iris/ReplicationSystem/NetRefHandleManager.cpp
--- a/Engine/Source/Runtime/Net/Iris/Private/Iris/ReplicationSystem/NetRefHandleManager.cpp
+++ b/Engine/Source/Runtime/Net/Iris/Private/Iris/ReplicationSystem/NetRefHandleManager.cpp
@@ -283,7 +283,8 @@ FInternalNetRefIndex FNetRefHandleManager::InternalCreateNetObject(const FNetRef
 				}
 				else
 				{
-					UE_LOG(LogIris, Warning, TEXT("NetRefHandleManager::InternalCreateNetObject Created NetObject %s for object already replicated as %s marked for destroy. This can lead to issues on clients. Only supported if using AsyncStopReplication."), *NetRefHandle.ToString(), *PrintObjectFromIndex(PendingDestroyInfo.InternalIndex));	
+					UE_LOG(LogIris, Verbose, TEXT("NetRefHandleManager::InternalCreateNetObject Created NetObject %s for object already replicated as %s marked for destroy."), *NetRefHandle.ToString(), *PrintObjectFromIndex(PendingDestroyInfo.InternalIndex));	
+					break;
 				}
 			}
 		}
@@ -574,6 +575,12 @@ void FNetRefHandleManager::InternalDestroyNetObject(FInternalNetRefIndex Interna
 	// Clear pointer to state buffer
 	ReplicatedObjectStateBuffers[InternalIndex] = nullptr;
 
+	// Remove cached creation header.
+	if (Data.bHasCachedCreationInfo)
+	{
+		CachedCreationHeaders.Remove(InternalIndex);
+	}
+
 	UE_NET_TRACE_NETHANDLE_DESTROYED(Data.RefHandle);
 
 	Data.Reset();
@@ -639,22 +646,44 @@ void FNetRefHandleManager::RemoveFromScope(FInternalNetRefIndex InternalIndex)
 	}
 }
 
-void FNetRefHandleManager::DestroyNetObject(FNetRefHandle RefHandle)
+void FNetRefHandleManager::DestroyNetObjectByIndex(FInternalNetRefIndex InternalIndex, bool bKeepInScopeAndMarkForFlush)
 {
-	FInternalNetRefIndex InternalIndex = RefHandleToInternalIndex.FindAndRemoveChecked(RefHandle);
 	if (ensure(AssignedInternalIndices.GetBit(InternalIndex)))
 	{
 		FReplicatedObjectData& Data = ReplicatedObjectData[InternalIndex];
-		check(Data.RefHandle == RefHandle);
 
-		// Remove mapping from global handle to internal index
-		NetHandleToInternalIndex.Remove(Data.NetHandle);
+		// If object already is marked for flush, we only need to remove it from scope.
+		if (!Data.bFlush)
+		{
+			// Just to catch misuse by bad code.
+			FInternalNetRefIndex FoundInternalIndex = RefHandleToInternalIndex.FindAndRemoveChecked(Data.RefHandle);
+			if (FoundInternalIndex != InternalIndex)
+			{
+				UE_LOG(LogIris, Error, TEXT("FNetRefHandleManager::DestroyNetObjectByIndex: Trying to destroy NetObject with conflicting InternalIndex: ( %u ) != ( %u) . Please fix calling code."), InternalIndex, FoundInternalIndex);
+				ensureMsgf(false, TEXT("Trying to destroy NetObject with conflicting InternalIndex"));
+				return;
+			}
+			
+			// Remove mapping from global handle to internal index
+			NetHandleToInternalIndex.Remove(Data.NetHandle);
 
-		// Remove from scopable objects if not already done
-		GlobalScopableInternalIndices.ClearBit(InternalIndex);
+			// We always defer the actual destroy
+			PendingDestroyInfos.Emplace(InternalIndex, Data.RefHandle);
+
+			// Before we can start replicating this RefHandle again, we need to first finish replication of previous incarnations of the object.
+			InternalIndicesReplicatingNetRefHandle.FindOrAdd(Data.RefHandle).Add(InternalIndex);
+		}
 
-		// We always defer the actual destroy
-		PendingDestroyInfos.Emplace(InternalIndex, RefHandle);
+		if (bKeepInScopeAndMarkForFlush)
+		{
+			// Mark as flush
+			Data.bFlush = true;
+		}
+		else
+		{
+			// Remove from scopable objects if not already done
+			GlobalScopableInternalIndices.ClearBit(InternalIndex);
+		}
 	}
 }
 
@@ -675,6 +704,17 @@ void FNetRefHandleManager::DestroyObjectsPendingDestroy()
 			FreedInternalIndices.Add(InternalIndex);
 
 			InternalDestroyNetObject(InternalIndex);
+
+			// Remove InternalIndex from array of objects replicating the same handle
+			if (TArray<FInternalNetRefIndex, TInlineAllocator<2>>* Entry = InternalIndicesReplicatingNetRefHandle.Find(PendingDestroyInfos[It].RefHandle))
+			{
+				Entry->Remove(InternalIndex);
+				if (Entry->IsEmpty())
+				{
+					InternalIndicesReplicatingNetRefHandle.Remove(PendingDestroyInfos[It].RefHandle);
+				}
+			}
+
 			PendingDestroyInfos.RemoveAtSwap(It);
 		}
 	}
@@ -1268,6 +1308,53 @@ void FNetRefHandleManager::AddReferencedObjects(FReferenceCollector& Collector)
 	}
 }
 
+void FNetRefHandleManager::StoreCachedCreationInfo(FInternalNetRefIndex InternalIndex, TUniquePtr<FNetObjectCreationHeader>&& Header)
+{
+	if (InternalIndex != InvalidInternalIndex)
+	{
+		FReplicatedObjectData& ObjectData = GetReplicatedObjectDataNoCheck(InternalIndex);
+		ObjectData.bHasCachedCreationInfo = 1U;
+		CachedCreationHeaders.Add(InternalIndex, MoveTemp(Header));
+	}
+}
+
+void FNetRefHandleManager::ClearCachedCreationInfo(FInternalNetRefIndex InternalIndex)
+{
+	if (InternalIndex != InvalidInternalIndex)
+	{
+		FReplicatedObjectData& ObjectData = GetReplicatedObjectDataNoCheck(InternalIndex);
+		ObjectData.bHasCachedCreationInfo = 0U;
+		CachedCreationHeaders.Remove(InternalIndex);
+	}
+}
+
+const FNetObjectCreationHeader* FNetRefHandleManager::GetCachedCreationInfo(FInternalNetRefIndex InternalIndex) const
+{
+	if (const TUniquePtr<const FNetObjectCreationHeader>* CachedHeader = CachedCreationHeaders.Find(InternalIndex))
+	{
+		return CachedHeader->Get();
+	}
+	else
+	{
+		return nullptr;
+	}
+}
+
+
+TConstArrayView<FInternalNetRefIndex> FNetRefHandleManager::GetInternalIndicesReplicatingNetRefHandle(FNetRefHandle NetRefHandle) const
+{
+	const TArray<FInternalNetRefIndex, TInlineAllocator<2>>* Entry = InternalIndicesReplicatingNetRefHandle.Find(NetRefHandle);
+	if (!Entry)
+	{
+		return MakeConstArrayView<FInternalNetRefIndex>(nullptr, 0);
+	}
+	else
+	{
+		return MakeConstArrayView(*Entry);
+	}
+}
+
+
 FNetRefHandle FNetRefHandleManager::MakeNetRefHandle(uint64 Id, uint32 ReplicationSystemId)
 {
 	check(Id <= FNetRefHandle::MaxId);
diff --git a/Engine/Source/Runtime/Net/Iris/Private/Iris/ReplicationSystem/NetRefHandleManager.h b/Engine/Source/Runtime/Net/Iris/Private/Iris/ReplicationSystem/NetRefHandleManager.h
--- a/Engine/Source/Runtime/Net/Iris/Private/Iris/ReplicationSystem/NetRefHandleManager.h
+++ b/Engine/Source/Runtime/Net/Iris/Private/Iris/ReplicationSystem/NetRefHandleManager.h
@@ -16,6 +16,7 @@
 
 #include "Iris/ReplicationState/ReplicationStateDescriptor.h"
 #include "Iris/ReplicationSystem/NetDependencyData.h"
+#include "Iris/ReplicationSystem/NetObjectFactory.h"
 #include "Iris/ReplicationSystem/NetObjectFactoryRegistry.h"
 #include "Iris/ReplicationSystem/NetRefHandle.h"
 #include "Iris/ReplicationSystem/NetRefHandleManagerTypes.h"
@@ -119,6 +120,7 @@ public:
 
 		uint16 bShouldPropagateChangedStates : 1U = 0;
 		uint16 bTearOff : 1U = 0;
+		uint16 bFlush : 1U = 0;
 		uint16 bIsDependentObject : 1U = 0;
 		uint16 bHasDependentObjects : 1U = 0;
 		uint16 bAllowDestroyInstanceFromRemote : 1U = 0;
@@ -222,7 +224,7 @@ public:
 	// Reset state of object tracked as internal index.
 	void ResetDestroyedStartupObject(FInternalNetRefIndex InternalIndex);
 	
-	void DestroyNetObject(FNetRefHandle Handle);
+	void DestroyNetObjectByIndex(FInternalNetRefIndex InternalIndex, bool bKeepInScopeAndMarkForFlush = false);
 
 	// Mark object as no longer scopable, object will be removed from scope for all connections
 	void RemoveFromScope(FInternalNetRefIndex InternalIndex);
@@ -434,6 +436,18 @@ public:
 	using FOnInternalNetRefIndicesFreed = TMulticastDelegate<void(const TConstArrayView<FInternalNetRefIndex>& FreedIndices)>;
 	FOnInternalNetRefIndicesFreed::RegistrationType& GetOnInternalNetRefIndicesFreedDelegate() const { return OnInternalNetRefIndicesFreed; };
 
+	/** Store cached creation info allowing us to send data for objects for which we no longer has an instance */
+	void StoreCachedCreationInfo(FInternalNetRefIndex, TUniquePtr<FNetObjectCreationHeader>&& Header);
+
+	/** Get cached creation header */
+	const FNetObjectCreationHeader* GetCachedCreationInfo(FInternalNetRefIndex InternalIndex) const;
+
+	/** Clear cached creation info  */
+	void ClearCachedCreationInfo(FInternalNetRefIndex InternalIndex);
+
+	/** Get internal indices if more than one internal indices are replicating the same NetRefHandle during flush. */
+	TConstArrayView<FInternalNetRefIndex> GetInternalIndicesReplicatingNetRefHandle(FNetRefHandle NetRefHandle) const;
+
 public:
 
 	// Iterate over all dependent objects and their dependent objects
@@ -610,6 +624,13 @@ private:
 
 	// Track NetBitArrays owned by the NetRefHandleManager so they can automatically be resized.
 	TArray<FNetBitArray*> OwnedNetBitArrays;
+
+	// Track previous use of internal indices that are still replicating data for the same handle to avoid having multiple internal objects replicating the same external object.
+	TMap<FNetRefHandle, TArray<FInternalNetRefIndex, TInlineAllocator<2>>> InternalIndicesReplicatingNetRefHandle;
+
+	// When we flush objects, we might need to defer sending creation info.
+	TMap<FInternalNetRefIndex, TUniquePtr<const UE::Net::FNetObjectCreationHeader>> CachedCreationHeaders;
+
 #if UE_SUPPORT_PARALLEL_IRIS
 	// Protects parallel access to NextHandleId and RefHandleToInternalIndex when allocating NetRefHandles
 	FTransactionallySafeCriticalSection HandleSerialCriticalSection;
diff --git a/Engine/Source/Runtime/Net/Iris/Private/Iris/ReplicationSystem/ObjectReplicationBridge.cpp b/Engine/Source/Runtime/Net/Iris/Private/Iris/ReplicationSystem/ObjectReplicationBridge.cpp
--- a/Engine/Source/Runtime/Net/Iris/Private/Iris/ReplicationSystem/ObjectReplicationBridge.cpp
+++ b/Engine/Source/Runtime/Net/Iris/Private/Iris/ReplicationSystem/ObjectReplicationBridge.cpp
@@ -435,43 +435,34 @@ void UObjectReplicationBridge::RemoveDestructionInfosForGroup(UE::Net::FNetObjec
 		FNetObjectGroup* Group = Groups->GetGroup(GroupHandle);
 		check(Group);
 
-		TArray<FNetRefHandle, TInlineAllocator<384>> ObjectsToRemove;	
 		TArray<FInternalNetRefIndex, TInlineAllocator<384>> ObjectIndicesToRemove;	
-		for (uint32 InternalObjectIndex : MakeArrayView(Group->Members))
+		for (FInternalNetRefIndex InternalObjectIndex : MakeArrayView(Group->Members))
 		{
 			if (NetRefHandleManager->IsDestroyedStartupObject(InternalObjectIndex))
 			{
 				const FNetRefHandleManager::FReplicatedObjectData& ObjectData = NetRefHandleManager->GetReplicatedObjectDataNoCheck(InternalObjectIndex);
 				if (StaticObjectsPendingDestroy.Remove(ObjectData.RefHandle))
 				{
-					ObjectsToRemove.Add(ObjectData.RefHandle);
 					ObjectIndicesToRemove.Add(InternalObjectIndex);
 				}
 			}
 		}
 
-		for (FNetRefHandle Handle : MakeArrayView(ObjectsToRemove))
+		// Destroy NetObjects and remove from WorldLocations
+		for (FInternalNetRefIndex InternalObjectIndex : MakeArrayView(ObjectIndicesToRemove))
 		{
-			NetRefHandleManager->DestroyNetObject(Handle);
-		}
-
-		for (FInternalNetRefIndex InternalReplicationIndex : MakeArrayView(ObjectIndicesToRemove))
-		{
-			WorldLocations.RemoveObjectInfoCache(InternalReplicationIndex);
+			NetRefHandleManager->DestroyNetObjectByIndex(InternalObjectIndex);
+			WorldLocations.RemoveObjectInfoCache(InternalObjectIndex);
 		}
 	}
 	else
 	{
-		// We should remove all destruction infos and objects
-		for (const auto& It : StaticObjectsPendingDestroy)
+		// Destroy NetObjects and remove from WorldLocations
+		for (const TPair<FNetRefHandle, UE::Net::Private::FStaticDestructionInfo>& It : StaticObjectsPendingDestroy)
 		{
-			NetRefHandleManager->DestroyNetObject(It.Key);
-		}
-
-		// Remove from WorldLocations
-		for (const auto& It : StaticObjectsPendingDestroy)
-		{
-			WorldLocations.RemoveObjectInfoCache(It.Value.InternalReplicationIndex);
+			const UE::Net::Private::FInternalNetRefIndex InternalObjectIndex = It.Value.InternalReplicationIndex;
+			NetRefHandleManager->DestroyNetObjectByIndex(InternalObjectIndex);
+			WorldLocations.RemoveObjectInfoCache(InternalObjectIndex);
 		}
 
 		StaticObjectsPendingDestroy.Empty();
@@ -493,7 +484,7 @@ void UObjectReplicationBridge::RemoveDestructionInfo(UE::Net::Private::FInternal
 	{
 		FWorldLocations& WorldLocations = GetReplicationSystem()->GetReplicationSystemInternal()->GetWorldLocations();
 
-		NetRefHandleManager->DestroyNetObject(ObjectData.RefHandle);
+		NetRefHandleManager->DestroyNetObjectByIndex(InternalObjectIndex);
 		WorldLocations.RemoveObjectInfoCache(InternalObjectIndex);
 	}
 }
@@ -581,6 +572,14 @@ UE::Net::FNetRefHandle UObjectReplicationBridge::GetReplicatedRefHandle(FNetHand
 	}
 
 	const FNetRefHandleManager::FReplicatedObjectData& ObjectData = NetRefHandleManager->GetReplicatedObjectDataNoCheck(ObjectInternalIndex);
+
+	// We do not consider objects pending stop due to flush or tearoff to be replicated to prevent unwanted API calls.
+	const bool bPendingStopDueToFlushOrTearOff = ObjectData.bTearOff || ObjectData.bFlush;
+	if (bPendingStopDueToFlushOrTearOff)
+	{
+		return FNetRefHandle::GetInvalid();
+	}
+
 	return ObjectData.RefHandle;
 }
 
@@ -1104,7 +1103,7 @@ void UObjectReplicationBridge::DestroyNetObjectFromRemote(UE::Net::Private::FInt
 	InternalDetachInstanceFromNetRefHandle(InternalObjectIndex);
 	
 	// Destroy the NetRefHandle
-	InternalDestroyNetObject(DestroyedObjectHandle);
+	InternalDestroyNetObject(InternalObjectIndex);
 }
 
 void UObjectReplicationBridge::SetSubObjectDormancyStatus(FNetRefHandle SubObjectRefHandle, FNetRefHandle OwnerRefHandle)
@@ -1229,7 +1228,7 @@ void UObjectReplicationBridge::RemoveCreationDependencyLink(FNetRefHandle Parent
 	GetReplicationSystem()->GetReplicationSystemInternal()->GetNetRefHandleManager().RemoveCreationDependency(Parent, Child);
 }
 
-bool UObjectReplicationBridge::WriteNetRefHandleCreationInfo(FReplicationBridgeSerializationContext& Context, FNetRefHandle Handle)
+bool UObjectReplicationBridge::WriteNetRefHandleCreationInfo(FReplicationBridgeSerializationContext& Context, FNetRefHandle Handle, UE::Net::Private::FInternalNetRefIndex InternalObjectIndex)
 {
 	UE_IRIS_PARALLEL_EXPR(FScopeLock ScopeLock(&WriteNetRefHandleCreationCS));
 	using namespace UE::Net;
@@ -1237,21 +1236,30 @@ bool UObjectReplicationBridge::WriteNetRefHandleCreationInfo(FReplicationBridgeS
 
 	check(!Context.bIsDestructionInfo);
 
+	// Validate that this is a legit object that we can write creation info for
+	if (!NetRefHandleManager->IsLocal(InternalObjectIndex))
+	{
+		UE_LOG_OBJECTREPLICATIONBRIDGE(Error, TEXT("WriteNetRefHandleCreationInfo: Cannot write creationinfo for %s ( InternalIndex: %u ), since it is no longer a valid object."), *Handle.ToString(), InternalObjectIndex);
+		ensureMsgf(false, TEXT("Cannot write creationinfo for %s ( InternalIndex: %u ), since it is no longer a valid object."), *Handle.ToString(), InternalObjectIndex);
+		return false;
+	}
+
+	const FNetRefHandleManager::FReplicatedObjectData& ObjectData = NetRefHandleManager->GetReplicatedObjectDataNoCheck(InternalObjectIndex);
+
 	// If there is a cached header for this handle, use it!
-	if (TUniquePtr<const FNetObjectCreationHeader>* CachedHeader = CachedCreationHeaders.Find(Handle))
+	if (const FNetObjectCreationHeader* Header = ObjectData.bHasCachedCreationInfo ? NetRefHandleManager->GetCachedCreationInfo(InternalObjectIndex) : nullptr)
 	{
 		//TODO: Right now cached creation data is only for destroyed or torn-off objects so the data can be considered const
 		//      If we eventually use it to store actively replicated objects, some creation data will want to be updated
 		//	    before being sent.  We should add a trait and a NetObjectFactory::UpdateCachedData to do so.
 
-		const FNetObjectCreationHeader* Header = CachedHeader->Get();
 		UNetObjectFactory* Factory = GetNetFactory(Header->GetNetFactoryId());
 		UE_LOG_BRIDGEID(LogIris, Verbose, TEXT("WriteNetRefHandleCreationInfo (cached) for %s : Factory: %s (id:%u) Header: %s"), *Handle.ToString(), *GetNameSafe(Factory), Header->GetNetFactoryId(), *Header->ToString());
 		return Factory->WriteHeader(Handle, Context.SerializationContext, Header);
 	}
-	
+
 	// Find the ProtocolId
-	const UE::Net::FReplicationProtocol* Protocol = GetReplicationSystem()->GetReplicationProtocol(Handle);
+	const UE::Net::FReplicationProtocol* Protocol = ObjectData.Protocol;
 	if (!Protocol)
 	{
 		UE_LOG_OBJECTREPLICATIONBRIDGE(Error, TEXT("WriteNetRefHandleCreationInfo: Cannot write creationinfo for %s, since protocol has been detached"), *Handle.ToString());
@@ -1259,14 +1267,6 @@ bool UObjectReplicationBridge::WriteNetRefHandleCreationInfo(FReplicationBridgeS
 		return false;
 	}
 
-	const FInternalNetRefIndex ObjectIndex = NetRefHandleManager->GetInternalIndex(Handle);
-	if (ObjectIndex == FNetRefHandleManager::InvalidInternalIndex)
-	{
-		UE_LOG_BRIDGEID(LogIris, Error, TEXT("WriteNetRefHandleCreationInfo: Cannot write creationinfo for %s: No internal index assigned to it"), *Handle.ToString());
-		return false;
-	}
-
-	const FNetRefHandleManager::FReplicatedObjectData& ObjectData = NetRefHandleManager->GetReplicatedObjectDataNoCheck(ObjectIndex);
 	const FNetObjectFactoryId FactoryId = ObjectData.NetFactoryId;
 
 	UNetObjectFactory* Factory = GetNetFactory(FactoryId);
@@ -1314,7 +1314,7 @@ bool UObjectReplicationBridge::CacheNetRefHandleCreationInfo(FNetRefHandle Handl
 		return false;
 	}
 	
-	CachedCreationHeaders.Add(Handle, MoveTemp(Header));
+	NetRefHandleManager->StoreCachedCreationInfo(ObjectIndex, MoveTemp(Header));
 	return true;
 }
 
@@ -1684,7 +1684,6 @@ void UObjectReplicationBridge::DetachSubObjectInstancesFromRemote(UE::Net::Priva
 void UObjectReplicationBridge::DetachInstance(FNetRefHandle RefHandle)
 {
 	UnregisterInstance(RefHandle);
-	CachedCreationHeaders.Remove(RefHandle);
 	Super::DetachInstance(RefHandle);
 }
 
diff --git a/Engine/Source/Runtime/Net/Iris/Private/Iris/ReplicationSystem/ReplicationBridge.cpp b/Engine/Source/Runtime/Net/Iris/Private/Iris/ReplicationSystem/ReplicationBridge.cpp
--- a/Engine/Source/Runtime/Net/Iris/Private/Iris/ReplicationSystem/ReplicationBridge.cpp
+++ b/Engine/Source/Runtime/Net/Iris/Private/Iris/ReplicationSystem/ReplicationBridge.cpp
@@ -102,8 +102,8 @@ void UReplicationBridge::GetInitialDependencies(FNetRefHandle Handle, FNetDepend
 
 void UReplicationBridge::CallPreSendUpdate(float DeltaSeconds)
 {
-	// Tear-off all handles pending tear-off
-	TearOffHandlesPendingTearOff();
+	// Prune stale async stop replication instances and tear-off all handles pending tear-off
+	PreUpdateHandlesPendingEndReplication();
 
 	PreSendUpdate();
 }
@@ -250,52 +250,30 @@ void UReplicationBridge::InternalDetachInstanceFromNetRefHandle(UE::Net::Private
 	}
 }
 
-void UReplicationBridge::InternalDestroyNetObject(FNetRefHandle Handle)
+void UReplicationBridge::InternalDestroyNetObject(UE::Net::Private::FInternalNetRefIndex InternalObjectIndex)
 {
 	using namespace UE::Net;
 	using namespace UE::Net::Private;
 
-	if (const FInternalNetRefIndex ObjectInternalIndex = NetRefHandleManager->GetInternalIndex(Handle))
+	if (InternalObjectIndex == FNetRefHandleManager::InvalidInternalIndex)
 	{
-		FReplicationSystemInternal* ReplicationSystemInternal = ReplicationSystem->GetReplicationSystemInternal();
-
-		FWorldLocations& WorldLocations = ReplicationSystemInternal->GetWorldLocations();
-		WorldLocations.RemoveObjectInfoCache(ObjectInternalIndex);
-
-		// Remove from handles to stop replicating as we might respawn the same object again in a later packet
-		HandlesToStopReplicating.Remove(Handle);
-
-		FNetRefHandleManager::FReplicatedObjectData& ObjectData = NetRefHandleManager->GetReplicatedObjectDataNoCheck(ObjectInternalIndex);
-
-		// At this point we can no longer instantiate objects and must explicitly clear out objects pending create from ReplicationWriter
-		if (ObjectData.bHasCachedCreationInfo && NetRefHandleManager->GetNetObjectRefCount(ObjectInternalIndex) > 0U)
-		{
-			// We need to explicitly notify all ReplicationWriters that we are destroying objects pending tearoff
-			// The handle will automatically be removed from HandlesPendingEndReplication after the next update
-			FReplicationConnections& Connections = ReplicationSystemInternal->GetConnections();
+		return;
+	}
 
-			auto NotifyDestroyedObjectPendingEndReplication = [&Connections, &ObjectInternalIndex](uint32 ConnectionId)
-			{
-				FReplicationConnection* Conn = Connections.GetConnection(ConnectionId);			
-				Conn->ReplicationWriter->NotifyDestroyedObjectPendingEndReplication(ObjectInternalIndex);
-			};
+	FReplicationSystemInternal* ReplicationSystemInternal = ReplicationSystem->GetReplicationSystemInternal();
 
-			const FNetBitArray& ValidConnections = Connections.GetValidConnections();
-			ValidConnections.ForAllSetBits(NotifyDestroyedObjectPendingEndReplication);					
+	FWorldLocations& WorldLocations = ReplicationSystemInternal->GetWorldLocations();
+	WorldLocations.RemoveObjectInfoCache(InternalObjectIndex);
 
-			// As we have release cached creation headers before calling InternalDestroyNetObject it is critical that we also clear out the flag.
-			ObjectData.bHasCachedCreationInfo = 0U;
-		}
-	}
+	// Remove from handles to stop replicating as we might respawn the same object again in a later packet
+	HandlesToStopReplicating.Remove(NetRefHandleManager->GetNetRefHandleFromInternalIndex(InternalObjectIndex));
 
-	NetRefHandleManager->DestroyNetObject(Handle);
+	NetRefHandleManager->DestroyNetObjectByIndex(InternalObjectIndex);
 }
 
-void UReplicationBridge::DestroyLocalNetHandle(FNetRefHandle Handle, EEndReplicationFlags EndReplicationFlags)
+void UReplicationBridge::InternalStartDestroyLocalObject(UE::Net::Private::FInternalNetRefIndex InternalReplicationIndex, EEndReplicationFlags EndReplicationFlags)
 {
-	UE_LOG_REPLICATIONBRIDGE(Verbose, TEXT("DestroyLocalNetHandle for %s | EndReplicationFlags: %s"), *NetRefHandleManager->PrintObjectFromNetRefHandle(Handle), *LexToString(EndReplicationFlags));
-
-	const UE::Net::Private::FInternalNetRefIndex InternalReplicationIndex = NetRefHandleManager->GetInternalIndex(Handle);
+	UE_LOG_REPLICATIONBRIDGE(Verbose, TEXT("InternalStartDestroyLocalObject for %s | EndReplicationFlags: %s"), *NetRefHandleManager->PrintObjectFromIndex(InternalReplicationIndex), *LexToString(EndReplicationFlags));
 
 	if (EnumHasAnyFlags(EndReplicationFlags, EEndReplicationFlags::DestroyNetHandle | EEndReplicationFlags::ClearNetPushId))
 	{
@@ -314,16 +292,41 @@ void UReplicationBridge::DestroyLocalNetHandle(FNetRefHandle Handle, EEndReplica
 	InternalDetachInstanceFromNetRefHandle(InternalReplicationIndex);
 
 	// Allow derived bridges to cleanup any instance info they have stored
+	// TODO: Call/DetachInstance should be deprecated and replace with direct call to UnRegisterInstance and NetObjectFactory method if still relevant.
+	FNetRefHandle Handle = NetRefHandleManager->GetNetRefHandleFromInternalIndex(InternalReplicationIndex);
 	CallDetachInstance(Handle);
 
-	// If the object is in any groups we need to remove it to make sure that we update filtering
-	GetReplicationSystem()->RemoveFromAllGroups(Handle);
+	if (EnumHasAnyFlags(EndReplicationFlags, EEndReplicationFlags::Flush))
+	{
+		// Initiate flush
+		const bool bKeepInScopeAndMarkForFlush = true;
+		NetRefHandleManager->DestroyNetObjectByIndex(InternalReplicationIndex, bKeepInScopeAndMarkForFlush);
+	}
+
+	// If we have any attached SubObjects, begin destroying them as well.
+	InternalStartDestroyLocalSubObjects(InternalReplicationIndex, EndReplicationFlags);
+}
+
+void UReplicationBridge::InternalFinishDestroyLocalObject(UE::Net::Private::FInternalNetRefIndex InternalReplicationIndex, EEndReplicationFlags EndReplicationFlags)
+{
+	// If the object are a member of any groups we need to remove it to make sure that we update filtering
+	GetReplicationSystem()->GetReplicationSystemInternal()->RemoveFromAllGroups(InternalReplicationIndex);
+
+	// If we have any attached SubObjects, finish destroy for them as well.
+	InternalFinishDestroyLocalSubObjects(InternalReplicationIndex, EndReplicationFlags);
+
+	// Finally Destroy NetObject and remove it from all scoping.
+	InternalDestroyNetObject(InternalReplicationIndex);
+}
+
+void UReplicationBridge::DestroyLocalNetHandle(FNetRefHandle Handle, EEndReplicationFlags EndReplicationFlags)
+{
+	UE_LOG_REPLICATIONBRIDGE(Verbose, TEXT("DestroyLocalNetHandle for %s | EndReplicationFlags: %s"), *NetRefHandleManager->PrintObjectFromNetRefHandle(Handle), *LexToString(EndReplicationFlags));
 
-	// If we have any attached SubObjects, tag them for destroy as well
-	InternalDestroySubObjects(Handle, EndReplicationFlags);
+	const UE::Net::Private::FInternalNetRefIndex InternalReplicationIndex = NetRefHandleManager->GetInternalIndex(Handle);
 
-	// Tell ReplicationSystem to destroy the handle
-	InternalDestroyNetObject(Handle);
+	InternalStartDestroyLocalObject(InternalReplicationIndex, EndReplicationFlags);
+	InternalFinishDestroyLocalObject(InternalReplicationIndex, EndReplicationFlags);
 }
 
 void UReplicationBridge::InternalAddSubObject(FNetRefHandle OwnerHandle, FNetRefHandle SubObjectHandle, FNetRefHandle InsertRelativeToSubObjectHandle, UE::Net::ESubObjectInsertionOrder InsertionOrder)
@@ -358,68 +361,77 @@ void UReplicationBridge::InternalAddSubObject(FNetRefHandle OwnerHandle, FNetRef
 	}
 }
 
-void UReplicationBridge::InternalDestroySubObjects(FNetRefHandle OwnerHandle, EEndReplicationFlags Flags)
+void UReplicationBridge::InternalStartDestroyLocalSubObjects(UE::Net::Private::FInternalNetRefIndex OwnerInternalIndex, EEndReplicationFlags Flags)
 {
 	using namespace UE::Net::Private;
 
 	// Destroy SubObjects
-	FInternalNetRefIndex OwnerInternalIndex = NetRefHandleManager->GetInternalIndex(OwnerHandle);
-	if (OwnerInternalIndex)
+	if (OwnerInternalIndex != FNetRefHandleManager::InvalidInternalIndex)
 	{
 		for (FInternalNetRefIndex SubObjectInternalIndex : NetRefHandleManager->GetChildSubObjects(OwnerInternalIndex))
 		{
-			FNetRefHandleManager::FReplicatedObjectData& SubObjectData = NetRefHandleManager->GetReplicatedObjectDataNoCheck(SubObjectInternalIndex);
+			// Tag subobject for destroy. The check against the scope is needed since the subobjects array might contain subobjects already pending destroy.
+			if (NetRefHandleManager->IsScopableIndex(SubObjectInternalIndex))
+			{
+				FNetRefHandleManager::FReplicatedObjectData& SubObjectData = NetRefHandleManager->GetReplicatedObjectDataNoCheck(SubObjectInternalIndex);
 
-			const FNetRefHandle SubObjectHandle = SubObjectData.RefHandle;
-				
+				SubObjectData.bPendingEndReplication = 1U;
+				if (!SubObjectData.bFlush)
+				{
+					UE_LOG_REPLICATIONBRIDGE(Verbose, TEXT("InternalStartDestroySubObjects %s"), *NetRefHandleManager->PrintObjectFromIndex(SubObjectInternalIndex));	
+					InternalStartDestroyLocalObject(SubObjectInternalIndex, Flags);				
+				}
+			}
+		}
+	}
+}
+
+void UReplicationBridge::InternalFinishDestroyLocalSubObjects(UE::Net::Private::FInternalNetRefIndex OwnerInternalIndex, EEndReplicationFlags Flags)
+{
+	using namespace UE::Net::Private;
+
+	// Destroy SubObjects
+	if (OwnerInternalIndex != FNetRefHandleManager::InvalidInternalIndex)
+	{
+		for (FInternalNetRefIndex SubObjectInternalIndex : NetRefHandleManager->GetChildSubObjects(OwnerInternalIndex))
+		{
 			// Tag subobject for destroy. The check against the scope is needed since the subobjects array might contain subobjects already pending destroy.
 			if (NetRefHandleManager->IsScopableIndex(SubObjectInternalIndex))
 			{
+				UE_LOG_REPLICATIONBRIDGE(Verbose, TEXT("InternalFinishDestroySubObjects %s"), *NetRefHandleManager->PrintObjectFromIndex(SubObjectInternalIndex));
+				FNetRefHandleManager::FReplicatedObjectData& SubObjectData = NetRefHandleManager->GetReplicatedObjectDataNoCheck(SubObjectInternalIndex);
 				SubObjectData.bPendingEndReplication = 1U;
-				UE_LOG_REPLICATIONBRIDGE(Verbose, TEXT("InternalDestroySubObjects %s"), *NetRefHandleManager->PrintObjectFromNetRefHandle(SubObjectHandle));
-				DestroyLocalNetHandle(SubObjectHandle, Flags);
+				InternalFinishDestroyLocalObject(SubObjectInternalIndex, Flags);
 			}
 		}
 	}
 }
 
-void UReplicationBridge::InternalRequestAsyncStopReplication(FNetRefHandle Handle, EEndReplicationFlags EndReplicationFlags)
+void UReplicationBridge::InternalRequestAsyncStopReplication(UE::Net::Private::FInternalNetRefIndex ObjectInternalIndex, EEndReplicationFlags EndReplicationFlags)
 {
 	using namespace UE::Net::Private;
 
-	UE_LOG_REPLICATIONBRIDGE(Verbose, TEXT("%hs: %s | EndFlags: %s"), __FUNCTION__, *NetRefHandleManager->PrintObjectFromNetRefHandle(Handle), ToCStr(LexToString(EndReplicationFlags)));
-	
-	FInternalNetRefIndex ObjectInternalIndex = NetRefHandleManager->GetInternalIndex(Handle);
-	if (ObjectInternalIndex != FNetRefHandleManager::InvalidInternalIndex)
+	if (ObjectInternalIndex != FNetRefHandleManager::InvalidInternalIndex && NetRefHandleManager->IsScopableIndex(ObjectInternalIndex))
 	{
+		UE_LOG_REPLICATIONBRIDGE(Verbose, TEXT("%hs: %s | EndFlags: %s"), __FUNCTION__, *NetRefHandleManager->PrintObjectFromIndex(ObjectInternalIndex), ToCStr(LexToString(EndReplicationFlags)));
+
 		FNetRefHandleManager::FReplicatedObjectData& ObjectData = NetRefHandleManager->GetReplicatedObjectDataNoCheck(ObjectInternalIndex);
+		AddPendingEndReplication(ObjectData.RefHandle, ObjectInternalIndex, EndReplicationFlags);
 
 		ObjectData.bIsAsyncStopping = true;
 
-		AddPendingEndReplication(Handle, EndReplicationFlags);
-
 		// If flush is required we need to wait to after scope update before we initialize async end replication, otherwise we can do it now.
 		if (!EnumHasAnyFlags(EndReplicationFlags, EEndReplicationFlags::Flush))
 		{
-			if (!ReplicationSystem->IsInGroup(ReplicationSystem->GetAsyncStopReplicationNetObjectGroup(), Handle))
+			if (!ReplicationSystem->IsInGroup(ReplicationSystem->GetAsyncStopReplicationNetObjectGroup(), ObjectData.RefHandle))
 			{
-				ReplicationSystem->AddToGroup(ReplicationSystem->GetAsyncStopReplicationNetObjectGroup(), Handle);
+				ReplicationSystem->AddToGroup(ReplicationSystem->GetAsyncStopReplicationNetObjectGroup(), ObjectData.RefHandle);
 			}
 		}
 
 		for (FInternalNetRefIndex SubObjectInternalIndex : NetRefHandleManager->GetChildSubObjects(ObjectInternalIndex))
 		{
-			FNetRefHandleManager::FReplicatedObjectData& SubObjectData = NetRefHandleManager->GetReplicatedObjectDataNoCheck(SubObjectInternalIndex);
-
-			const FNetRefHandle SubObjectHandle = SubObjectData.RefHandle;
-
-			// Tag subobject for async stop replication. The check against the scope is needed since the subobjects array might contain destroyed subobjects
-			if (NetRefHandleManager->IsScopableIndex(SubObjectInternalIndex))
-			{
-				//UE_LOG_REPLICATIONBRIDGE(Verbose, TEXT("InternalRequestAsyncStopReplicationForSubObjects %s"), *NetRefHandleManager->PrintObjectFromNetRefHandle(SubObjectHandle));
-				//TODO: Would be faster if we didn't do a recursive call here
-				InternalRequestAsyncStopReplication(SubObjectHandle, EndReplicationFlags);
-			}
+			InternalRequestAsyncStopReplication(SubObjectInternalIndex, EndReplicationFlags);
 		}
 	}
 }
@@ -457,13 +469,13 @@ void UReplicationBridge::StopReplicatingNetRefHandle(FNetRefHandle Handle, EEndR
 		{
 			ObjectData.InternalDetachReason = EInternalDetachReason::TornOff;
 
-			// Add handle to list of objects pending EndReplication indicate that it should be destroyed during next update
-			// We need to do this to cover the case where the torn off object not yet has been added to the scope.
-			AddPendingEndReplication(Handle, EndReplicationFlags);
-
 			// We do however copy the final state data and mark object to stop propagating state changes
 			InternalTearOff(Handle);
 
+			// Add handle to list of objects pending EndReplication indicate that it should be destroyed during next update
+			// We need to do this to cover the case where the torn off object not yet has been added to the scope.
+			AddPendingEndReplication(Handle, InternalReplicationIndex, EndReplicationFlags);
+
 			// Detach instance as we must assume that we should not access the object after this call.
 			ObjectData.bPendingEndReplication = 1U;
 			InternalDetachInstanceFromNetRefHandle(InternalReplicationIndex);
@@ -483,28 +495,25 @@ void UReplicationBridge::StopReplicatingNetRefHandle(FNetRefHandle Handle, EEndR
 			
 			if (bAsyncStopRequested)
 			{
-				// Defer stop replication until object is no longer scoped by any connections. This allows to restart replication on the object seamlessly if a start was requested for the same object while the objects are still relevant.
-
-
 				if (bFlushRequested)
 				{
-					InternalFlushStateData(Handle);
+					InternalFlushStateData(InternalReplicationIndex);
 				}
 
-				// Note: When using async stop we only detach the instance when it is no longer referenced by every client.
-				InternalRequestAsyncStopReplication(Handle, EndReplicationFlags);
+				// Note: When using async stop we only detach the instance when it is no longer referenced by any client.
+				InternalRequestAsyncStopReplication(InternalReplicationIndex, EndReplicationFlags);
 			}
 			else if (bFlushRequested)
 			{
-				// Defer destroy until after scope update to allow create/destroy on same frame
-				AddPendingEndReplication(Handle, EndReplicationFlags);
-			
 				// Capture final state
-				InternalFlushStateData(Handle);
+				InternalFlushStateData(InternalReplicationIndex);
+
+				// Object will finalize the flush-destroy Immediately after the next NetUpdate.
+				AddPendingEndReplication(Handle, InternalReplicationIndex, EndReplicationFlags, EPendingEndReplicationImmediate::Yes);
 
-				// Detach instance as we must assume that we should not access the object after this call.
+				// After this call, the NetObject is only accessible through its InternalReplicationIndex.
 				ObjectData.bPendingEndReplication = 1U;
-				InternalDetachInstanceFromNetRefHandle(InternalReplicationIndex);
+				InternalStartDestroyLocalObject(InternalReplicationIndex, EndReplicationFlags);
 			}
 			else
 			{
@@ -515,7 +524,6 @@ void UReplicationBridge::StopReplicatingNetRefHandle(FNetRefHandle Handle, EEndR
 	}
 	else
 	{
-
 		// While we are inside ReceiveUpdate, queue stop replication requests instead of immediately stopping replication
 		// This allows us the apply any received updates before we cut off this object
 		if (IsInReceiveUpdate())
@@ -545,13 +553,29 @@ void UReplicationBridge::StopReplicatingNetRefHandle(FNetRefHandle Handle, EEndR
 	}
 }
 
-
-
-void UReplicationBridge::TearOffHandlesPendingTearOff()
+void UReplicationBridge::PreUpdateHandlesPendingEndReplication()
 {
-	// Initiate tear off
+	using namespace UE::Net;
+	using namespace UE::Net::Private;
+
 	for (const FPendingEndReplicationInfo& Info : MakeArrayView(HandlesPendingEndReplication))
 	{
+		// Due to async stop replication not detaching instances in the call to StopReplication  we need to explicitly check for invalidated instances as objects might have been marked for garbage but GC has not been run.
+		if (FInternalNetRefIndex ObjectInternalIndex = NetRefHandleManager->GetInternalIndex(Info.Handle))
+		{
+			if (!IsValid(NetRefHandleManager->GetReplicatedInstances()[ObjectInternalIndex]))
+			{
+				FNetRefHandleManager::FReplicatedObjectData& ObjectData = NetRefHandleManager->GetReplicatedObjectDataNoCheck(ObjectInternalIndex);
+				if (ObjectData.InstanceProtocol)
+				{
+					// Detach instance as we should not access the instance.
+					ObjectData.bPendingEndReplication = 1U;
+					InternalDetachInstanceFromNetRefHandle(ObjectInternalIndex);
+				}
+			}
+		}
+
+		// Tear off handle pending tear off.
 		if (EnumHasAnyFlags(Info.EndReplicationFlags, EEndReplicationFlags::TearOff))
 		{
 			InternalTearOff(Info.Handle);
@@ -564,23 +588,35 @@ void UReplicationBridge::UpdateHandlesPendingEndReplication()
 	using namespace UE::Net;
 	using namespace UE::Net::Private;
 
+	const FNetBitArray& AssignedInternalIndices = NetRefHandleManager->GetAssignedInternalIndices();
+
 	TArray<FPendingEndReplicationInfo, TInlineAllocator<32>> ObjectsStillPendingEndReplication;
 	for (FPendingEndReplicationInfo Info : MakeArrayView(HandlesPendingEndReplication))
 	{
-		if (FInternalNetRefIndex ObjectInternalIndex = NetRefHandleManager->GetInternalIndex(Info.Handle))
+		const FInternalNetRefIndex ObjectInternalIndex = Info.ObjectIndex;
+		if (NetRefHandleManager->GetAssignedInternalIndices().GetBit(ObjectInternalIndex))
 		{
+			FNetRefHandleManager::FReplicatedObjectData& ObjectData = NetRefHandleManager->GetReplicatedObjectDataNoCheck(ObjectInternalIndex);
+			check(ObjectData.RefHandle == Info.Handle);
+
 			// Immediate destroy or objects that are no longer are referenced by any connections are destroyed
 			if (NetRefHandleManager->GetNetObjectRefCount(ObjectInternalIndex) == 0U || Info.Immediate == EPendingEndReplicationImmediate::Yes)
 			{
-				FNetRefHandleManager::FReplicatedObjectData& ObjectData = NetRefHandleManager->GetReplicatedObjectDataNoCheck(ObjectInternalIndex);
 				ObjectData.bPendingEndReplication = 1U;
-				DestroyLocalNetHandle(Info.Handle, Info.EndReplicationFlags);
+				if (ObjectData.bFlush)
+				{
+					// If we have initiated a flush we just need to finish the destroy
+					InternalFinishDestroyLocalObject(Info.ObjectIndex, Info.EndReplicationFlags);
+				}
+				else
+				{
+					DestroyLocalNetHandle(Info.Handle, Info.EndReplicationFlags);
+				}
 			}
 			else
 			{
 				if (EnumHasAnyFlags(Info.EndReplicationFlags, EEndReplicationFlags::AllowAsyncStopReplication))
 				{
-					FNetRefHandleManager::FReplicatedObjectData& ObjectData = NetRefHandleManager->GetReplicatedObjectDataNoCheck(ObjectInternalIndex);
 					// We keep the object in scope but we mark it a not replicated which will stop replicating it for all clients
 					if (!ReplicationSystem->IsInGroup(ReplicationSystem->GetAsyncStopReplicationNetObjectGroup(), Info.Handle))
 					{
@@ -590,7 +626,6 @@ void UReplicationBridge::UpdateHandlesPendingEndReplication()
 				else if (NetRefHandleManager->IsScopableIndex(ObjectInternalIndex))
 				{
 					// If the object is still in scope remove it from scope as objects pending EndReplication should not be added to new connections after the first update.
-
 					// Mark object and subobjects as no longer scopeable, and that we should not propagate changed states
 					NetRefHandleManager->RemoveFromScope(ObjectInternalIndex);
 					for (FInternalNetRefIndex SubObjectIndex : NetRefHandleManager->GetSubObjects(ObjectInternalIndex))
@@ -600,7 +635,7 @@ void UReplicationBridge::UpdateHandlesPendingEndReplication()
 				}
 			
 				// Keep object in the pending EndReplication-list until the object is no longer referenced by any ReplicationWriter
-				ObjectsStillPendingEndReplication.Add(FPendingEndReplicationInfo(Info.Handle, Info.EndReplicationFlags, EPendingEndReplicationImmediate::No));
+				ObjectsStillPendingEndReplication.Add(FPendingEndReplicationInfo(Info.Handle, ObjectInternalIndex, Info.EndReplicationFlags, EPendingEndReplicationImmediate::No));
 			}
 		}
 	}
@@ -611,21 +646,24 @@ void UReplicationBridge::UpdateHandlesPendingEndReplication()
 	CSV_CUSTOM_STAT(Iris, NumHandlesPendingEndRepliction, (float)HandlesPendingEndReplication.Num(), ECsvCustomStatOp::Set);
 }
 
-void UReplicationBridge::AddPendingEndReplication(FNetRefHandle Handle, EEndReplicationFlags DestroyFlags, EPendingEndReplicationImmediate Immediate)
+void UReplicationBridge::AddPendingEndReplication(FNetRefHandle Handle, UE::Net::Private::FInternalNetRefIndex InternalObjectIndex, EEndReplicationFlags DestroyFlags, EPendingEndReplicationImmediate Immediate)
 {
 	if (ensure(EnumHasAnyFlags(DestroyFlags, EEndReplicationFlags::Flush | EEndReplicationFlags::TearOff | EEndReplicationFlags::AllowAsyncStopReplication)))
 	{
-		if (!HandlesPendingEndReplication.FindByPredicate([&](const FPendingEndReplicationInfo& Info){ return Info.Handle == Handle; }))
+		if (!HandlesPendingEndReplication.FindByPredicate([&InternalObjectIndex](const FPendingEndReplicationInfo& Info){ return Info.ObjectIndex == InternalObjectIndex; }))
 		{
-			HandlesPendingEndReplication.Add(FPendingEndReplicationInfo(Handle, DestroyFlags, Immediate));
+			HandlesPendingEndReplication.Add(FPendingEndReplicationInfo(Handle, InternalObjectIndex, DestroyFlags, Immediate));
 		}
 	}
 }
 
 bool UReplicationBridge::IsPendingStopReplication(FNetRefHandle Handle) const
 {
-	const int32 Index = HandlesPendingEndReplication.FindLastByPredicate([Handle](const FPendingEndReplicationInfo& Info){ return Info.Handle == Handle; });
-	return Index != INDEX_NONE;
+	const int32 Index = HandlesPendingEndReplication.FindLastByPredicate([Handle](const FPendingEndReplicationInfo& Info)
+	{ 
+		return Info.Handle == Handle;
+	});
+	return Index != INDEX_NONE && EnumHasAnyFlags(HandlesPendingEndReplication[Index].EndReplicationFlags, EEndReplicationFlags::Flush | EEndReplicationFlags::AllowAsyncStopReplication);
 }
 
 bool UReplicationBridge::RestartOrStopObjectPendingAsyncStopReplication(FNetRefHandle Handle, bool bForceStopReplication)
@@ -636,9 +674,10 @@ bool UReplicationBridge::RestartOrStopObjectPendingAsyncStopReplication(FNetRefH
 	bool bRestarted = false;
 
 	const int32 Index = HandlesPendingEndReplication.FindLastByPredicate([Handle](const FPendingEndReplicationInfo& Info){ return Info.Handle == Handle; });
+	// We allow canceling StopReplication for objects that are marked for AsyncStopReplication
 	if (Index != INDEX_NONE && EnumHasAnyFlags(HandlesPendingEndReplication[Index].EndReplicationFlags, EEndReplicationFlags::AllowAsyncStopReplication))
 	{
-		UE_LOG_REPLICATIONBRIDGE(Verbose, TEXT("RestartOrStopObjectPendingAsyncStopReplication: %s ForceStop: %u"), *PrintObjectFromNetRefHandle(Handle), bForceStopReplication ? 1U : 0U);
+		UE_LOG_REPLICATIONBRIDGE(Verbose, TEXT("RestartOrStopObjectPendingAsyncStopReplication: %s ForceStop: %u Flags:%s"), *PrintObjectFromNetRefHandle(Handle), bForceStopReplication ? 1U : 0U, *LexToString(HandlesPendingEndReplication[Index].EndReplicationFlags));
 
 		const uint32 InternalObjectIndex = NetRefHandleManager->GetInternalIndex(HandlesPendingEndReplication[Index].Handle);
 		if (InternalObjectIndex != FNetRefHandleManager::InvalidInternalIndex)
@@ -655,8 +694,8 @@ bool UReplicationBridge::RestartOrStopObjectPendingAsyncStopReplication(FNetRefH
 			else
 			{
 				ObjectData.bIsAsyncStopping = false;
-
 				bRestarted = true;
+
 				if (NetRefHandleManager->IsDestroyedStartupObject(InternalObjectIndex))
 				{
 					// Need to cleanse destruction info if we are going to cancel async stop replication.
@@ -664,11 +703,18 @@ bool UReplicationBridge::RestartOrStopObjectPendingAsyncStopReplication(FNetRefH
 					NetRefHandleManager->ResetDestroyedStartupObject(InternalObjectIndex);
 					RemoveDestructionInfo(DestructionInfoInternalIndex);
 				}
+				
+				// Remove cached creation info.
+				NetRefHandleManager->ClearCachedCreationInfo(InternalObjectIndex);
 			}
 		}
 
 		HandlesPendingEndReplication.RemoveAt(Index);
-		ReplicationSystem->RemoveFromGroup(ReplicationSystem->GetAsyncStopReplicationNetObjectGroup(), Handle);
+		const FNetObjectGroupHandle AsyncStopReplicationNetObjectGroupHandle = ReplicationSystem->GetAsyncStopReplicationNetObjectGroup();
+		if (ReplicationSystem->IsInGroup(AsyncStopReplicationNetObjectGroupHandle, Handle))
+		{
+			ReplicationSystem->RemoveFromGroup(AsyncStopReplicationNetObjectGroupHandle, Handle);
+		}
 
 		return bRestarted;
 	}
@@ -676,7 +722,7 @@ bool UReplicationBridge::RestartOrStopObjectPendingAsyncStopReplication(FNetRefH
 	return Index == INDEX_NONE;
 }
 
-void UReplicationBridge::InternalFlushStateData(UE::Net::FNetSerializationContext& SerializationContext, UE::Net::Private::FChangeMaskCache& ChangeMaskCache, UE::Net::FNetBitStreamWriter& ChangeMaskWriter, uint32 InternalObjectIndex)
+void UReplicationBridge::InternalFlushStateData(UE::Net::FNetSerializationContext& SerializationContext, UE::Net::Private::FChangeMaskCache& ChangeMaskCache, UE::Net::FNetBitStreamWriter& ChangeMaskWriter, UE::Net::Private::FInternalNetRefIndex InternalObjectIndex)
 {
 	using namespace UE::Net;
 	using namespace UE::Net::Private;
@@ -718,14 +764,13 @@ void UReplicationBridge::InternalFlushStateData(UE::Net::FNetSerializationContex
 	}
 }
 
-void UReplicationBridge::InternalFlushStateData(FNetRefHandle Handle)
+void UReplicationBridge::InternalFlushStateData(UE::Net::Private::FInternalNetRefIndex InternalObjectIndex)
 {
 	using namespace UE::Net;
 	using namespace UE::Net::Private;
 
 	IRIS_PROFILER_SCOPE(InternalFlushStateData);
 
-	const uint32 InternalObjectIndex = NetRefHandleManager->GetInternalIndex(Handle);
 	if (InternalObjectIndex == FNetRefHandleManager::InvalidInternalIndex)
 	{
 		return;
@@ -851,7 +896,19 @@ void UReplicationBridge::InternalTearOff(FNetRefHandle Handle)
 
 bool UReplicationBridge::IsReplicatedHandle(FNetRefHandle Handle) const
 {
-	return Handle.IsValid() && NetRefHandleManager->IsNetRefHandleAssigned(Handle);
+	using namespace UE::Net::Private;
+
+	const FInternalNetRefIndex InternalObjectIndex = NetRefHandleManager->GetInternalIndex(Handle);
+	if (InternalObjectIndex == FNetRefHandleManager::InvalidInternalIndex)
+	{
+		return false;
+	}
+
+	// We do not consider objects pending stop due to flush or tearoff to be replicated to prevent unwanted API calls.
+	FNetRefHandleManager::FReplicatedObjectData& ObjectData = NetRefHandleManager->GetReplicatedObjectDataNoCheck(InternalObjectIndex);
+	const bool bPendingStopDueToFlushOrTearOff = ObjectData.bTearOff || ObjectData.bFlush;
+
+	return !bPendingStopDueToFlushOrTearOff;
 }
 
 void UReplicationBridge::SetNetPushIdOnFragments(const TArrayView<const UE::Net::FReplicationFragment*const>& Fragments, const UE::Net::Private::FNetPushObjectHandle& PushHandle)
diff --git a/Engine/Source/Runtime/Net/Iris/Private/Iris/ReplicationSystem/ReplicationSystem.cpp b/Engine/Source/Runtime/Net/Iris/Private/Iris/ReplicationSystem/ReplicationSystem.cpp
--- a/Engine/Source/Runtime/Net/Iris/Private/Iris/ReplicationSystem/ReplicationSystem.cpp
+++ b/Engine/Source/Runtime/Net/Iris/Private/Iris/ReplicationSystem/ReplicationSystem.cpp
@@ -848,6 +848,26 @@ public:
 	}
 }; // end class FReplicationSystemImpl
 
+void FReplicationSystemInternal::RemoveFromAllGroups(FInternalNetRefIndex ObjectInternalIndex)
+{
+	uint32 NumGroupMemberShips = 0;
+	if (AreFilterChangesBlocked())
+	{
+		ensureMsgf(false, TEXT("Setting filter conditions is not yet supported during this operation. Filter condition on %s failed."), *NetRefHandleManager.PrintObjectFromIndex(ObjectInternalIndex));
+		return;
+	}
+
+	// We copy the membership array as it is modified during removal
+	TArray<FNetObjectGroupHandle> CopiedGroupHandles;
+	Groups.GetGroupHandlesOfNetObject(ObjectInternalIndex, CopiedGroupHandles);
+
+	for (FNetObjectGroupHandle GroupHandle : CopiedGroupHandles)
+	{
+		Groups.RemoveFromGroup(GroupHandle, ObjectInternalIndex);
+		Filtering.NotifyObjectRemovedFromGroup(GroupHandle, ObjectInternalIndex);
+	}	
+}
+
 } // end namespace UE::Net::Private
 
 UReplicationSystem::UReplicationSystem()
@@ -1805,27 +1825,9 @@ void UReplicationSystem::RemoveFromAllGroups(FNetRefHandle Handle)
 	using namespace UE::Net::Private;
 
 	FNetRefHandleManager& NetRefHandleManager = Impl->ReplicationSystemInternal.GetNetRefHandleManager();
-	FNetObjectGroups& Groups = Impl->ReplicationSystemInternal.GetGroups();
-	FReplicationFiltering& Filtering = Impl->ReplicationSystemInternal.GetFiltering();
-
-	uint32 NumGroupMemberShips = 0;
 	const FInternalNetRefIndex ObjectInternalIndex = NetRefHandleManager.GetInternalIndex(Handle);
 
-	if (Impl->ReplicationSystemInternal.AreFilterChangesBlocked())
-	{
-		ensureMsgf(false, TEXT("Setting filter conditions is not yet supported during this operation. Filter condition on %s (%s) failed."), *GetNameSafe(NetRefHandleManager.GetReplicatedObjectInstance(ObjectInternalIndex)), *Handle.ToString());
-		return;
-	}
-
-	// We copy the membership array as it is modified during removal
-	TArray<FNetObjectGroupHandle> CopiedGroupHandles;
-	Groups.GetGroupHandlesOfNetObject(ObjectInternalIndex, CopiedGroupHandles);
-
-	for (FNetObjectGroupHandle GroupHandle : CopiedGroupHandles)
-	{
-		Groups.RemoveFromGroup(GroupHandle, ObjectInternalIndex);
-		Filtering.NotifyObjectRemovedFromGroup(GroupHandle, ObjectInternalIndex);
-	}	
+	Impl->ReplicationSystemInternal.RemoveFromAllGroups(ObjectInternalIndex);
 }
 
 bool UReplicationSystem::IsInGroup(FNetObjectGroupHandle GroupHandle, FNetRefHandle Handle) const
@@ -2094,7 +2096,10 @@ void UReplicationSystem::SetIsNetTemporary(FNetRefHandle Handle)
 void UReplicationSystem::TearOffNextUpdate(FNetRefHandle Handle)
 {
 	constexpr EEndReplicationFlags DestroyFlags = EEndReplicationFlags::TearOff | EEndReplicationFlags::ClearNetPushId;
-	Impl->ReplicationSystemInternal.GetReplicationBridge()->AddPendingEndReplication(Handle, DestroyFlags);
+	if (const uint32 InternalObjectIndex = Impl->ReplicationSystemInternal.GetNetRefHandleManager().GetInternalIndex(Handle))
+	{
+		Impl->ReplicationSystemInternal.GetReplicationBridge()->AddPendingEndReplication(Handle, InternalObjectIndex, DestroyFlags);
+	}
 }
 
 void UReplicationSystem::ForceNetUpdate(FNetRefHandle Handle)
diff --git a/Engine/Source/Runtime/Net/Iris/Private/Iris/ReplicationSystem/ReplicationSystemInternal.h b/Engine/Source/Runtime/Net/Iris/Private/Iris/ReplicationSystem/ReplicationSystemInternal.h
--- a/Engine/Source/Runtime/Net/Iris/Private/Iris/ReplicationSystem/ReplicationSystemInternal.h
+++ b/Engine/Source/Runtime/Net/Iris/Private/Iris/ReplicationSystem/ReplicationSystemInternal.h
@@ -205,6 +205,8 @@ public:
 	}
 #endif // UE_SUPPORT_PARALLEL_IRIS
 
+void RemoveFromAllGroups(UE::Net::Private::FInternalNetRefIndex ObjectInternalIndex);
+
 private:
 	FReplicationProtocolManager ReplicationProtocolManager;
 	FNetRefHandleManager NetRefHandleManager;
diff --git a/Engine/Source/Runtime/Net/Iris/Private/Iris/ReplicationSystem/ReplicationWriter.cpp b/Engine/Source/Runtime/Net/Iris/Private/Iris/ReplicationSystem/ReplicationWriter.cpp
--- a/Engine/Source/Runtime/Net/Iris/Private/Iris/ReplicationSystem/ReplicationWriter.cpp
+++ b/Engine/Source/Runtime/Net/Iris/Private/Iris/ReplicationSystem/ReplicationWriter.cpp
@@ -594,7 +594,8 @@ void FReplicationWriter::StartReplication(uint32 InternalIndex)
 	Info.IsDestructionInfo = bIsDestructionInfo;
 	Info.IsCreationConfirmed = 0U;
 	Info.TearOff = Data.bTearOff;
-	Info.FlushFlags = GetDefaultFlushFlags();
+	// If object is marked for flush we set the FlushFlags_FlushState in addition to the default ones.
+	Info.FlushFlags = GetDefaultFlushFlags() | (Data.bFlush ? EFlushFlags::FlushFlags_FlushState : EFlushFlags::FlushFlags_None);
 	Info.SubObjectPendingDestroy = 0U;
 	Info.IsDeltaCompressionEnabled = BaselineManager->GetDeltaCompressionStatus(InternalIndex) == ENetObjectDeltaCompressionStatus::Allow ? 1U : 0U;
 	Info.LastAckedBaselineIndex = FDeltaCompressionBaselineManager::InvalidBaselineIndex;
@@ -2155,9 +2156,9 @@ uint32 FReplicationWriter::WriteRootObjectsPendingDestroy(FNetSerializationConte
 			if (Info.TearOff)
 			{
 				// Check that all subobjects are ready to be torn off too.
-				const bool bAllSubObjectsAreReadyToBeTornOff = [this]()
+				const bool bAllSubObjectsAreReadyToBeTornOff = [this, InternalIndex]()
 				{
-					for (const FInternalNetRefIndex SubObjectIndex : NetRefHandleManager->GetSubObjectInternalIndicesView())
+					for (const FInternalNetRefIndex SubObjectIndex : NetRefHandleManager->GetSubObjects(InternalIndex))
 					{
 						FReplicationInfo& SubObjectInfo = ReplicatedObjects[SubObjectIndex];
 						if (SubObjectInfo.GetState() == EReplicatedObjectState::WaitOnFlush)
@@ -2183,6 +2184,8 @@ uint32 FReplicationWriter::WriteRootObjectsPendingDestroy(FNetSerializationConte
 					ObjectsWithDirtyChanges.SetBit(InternalIndex);
 					Info.HasDirtySubObjects = 1;
 				}
+				// Either we are waiting for flush or we are PendingTearOff.
+				continue;
 			}
 			else
 			{
@@ -2352,6 +2355,29 @@ bool FReplicationWriter::CanSendObject(uint32 InternalIndex) const
 	const FReplicationInfo& Info = GetReplicationInfo(InternalIndex);
 	const EReplicatedObjectState State = Info.GetState();
 
+	if (State == EReplicatedObjectState::PendingCreate)
+	{
+		// We cannot replicate the same netrefhandle using different internal indices.
+		const FNetRefHandleManager::FReplicatedObjectData& Data = NetRefHandleManager->GetReplicatedObjectDataNoCheck(InternalIndex);		
+		for (const FInternalNetRefIndex& PendingDestroyIndex : NetRefHandleManager->GetInternalIndicesReplicatingNetRefHandle(Data.RefHandle))
+		{
+			// If we find the current InternalIndex we know that we can start replicating object identified by it.
+			if (PendingDestroyIndex == InternalIndex)
+			{
+				break;
+			}
+
+			const FReplicationInfo& PendingDestroyInfo = GetReplicationInfo(PendingDestroyIndex);
+			const EReplicatedObjectState PendingDestroyState = PendingDestroyInfo.GetState();
+				
+			if (PendingDestroyState != EReplicatedObjectState::Invalid)
+			{
+				UE_NET_REPLICATIONWRITER_WARN_IF_CANNOT_SEND(TEXT("Waiting on previous instance %s to finish replication."), *NetRefHandleManager->PrintObjectFromIndex(PendingDestroyIndex));
+				return false;
+			}				
+		}
+	}
+
 	// Currently we do wait for CreateConfirmation before sending more data
 	// We might want to change this and allow "bombing" creation info until we get confirmation to minimize latency
 	// We also prevent objects from being transmitted if they are waiting on destroy/tear-off confirmation or cancelling destroy.
@@ -2785,7 +2811,7 @@ FReplicationWriter::EWriteObjectStatus FReplicationWriter::WriteObjectAndSubObje
 					}
 					else
 					{
-						bWriteSuccess = FBridgeSerialization::WriteNetRefHandleCreationInfo(ReplicationBridge, BridgeContext, NetRefHandle);
+						bWriteSuccess = FBridgeSerialization::WriteNetRefHandleCreationInfo(ReplicationBridge, BridgeContext, NetRefHandle, InternalIndex);
 					}
 
 					// We need to send creation info, so if we fail we skip this object for now
diff --git a/Engine/Source/Runtime/Net/Iris/Public/Iris/ReplicationSystem/ObjectReplicationBridge.h b/Engine/Source/Runtime/Net/Iris/Public/Iris/ReplicationSystem/ObjectReplicationBridge.h
--- a/Engine/Source/Runtime/Net/Iris/Public/Iris/ReplicationSystem/ObjectReplicationBridge.h
+++ b/Engine/Source/Runtime/Net/Iris/Public/Iris/ReplicationSystem/ObjectReplicationBridge.h
@@ -568,8 +568,9 @@ private:
 	 * Write data required to instantiate NetObject remotely to bitstream.
 	 * @param Context The serialization context parameters.
 	 * @param Handle The handle of the object to write creation data for.
+	 * @param InternalObjectIndex the internalindex we are currently writing creation info for.
 	 */
-	bool WriteNetRefHandleCreationInfo(FReplicationBridgeSerializationContext& Context, FNetRefHandle Handle);
+	bool WriteNetRefHandleCreationInfo(FReplicationBridgeSerializationContext& Context, FNetRefHandle Handle, UE::Net::Private::FInternalNetRefIndex InternalObjectIndex);
 
 	/**
 	 * Client only function that executes the EndReplication request by the authority
@@ -658,9 +659,6 @@ private:
 	// Type stats
 	TMap<FName, FName> ClassesWithTypeStats;
 
-	// When we flush objects, we might need to defer sending creation info.
-	TMap<FNetRefHandle, TUniquePtr<const UE::Net::FNetObjectCreationHeader>> CachedCreationHeaders;
-
 	// Track the static objects that are destroyed so clients can be told to delete them after loading the container
 	TMap<FNetRefHandle, UE::Net::Private::FStaticDestructionInfo> StaticObjectsPendingDestroy;
 
@@ -777,9 +775,9 @@ private:
 	friend class FReplicationReader;
 	friend UNetObjectFactory;
 
-	static bool WriteNetRefHandleCreationInfo(UObjectReplicationBridge* Bridge, FReplicationBridgeSerializationContext& Context, FNetRefHandle Handle)
+	static bool WriteNetRefHandleCreationInfo(UObjectReplicationBridge* Bridge, FReplicationBridgeSerializationContext& Context, FNetRefHandle Handle, FInternalNetRefIndex InternalObjectIndex)
 	{
-		return Bridge->WriteNetRefHandleCreationInfo(Context, Handle);
+		return Bridge->WriteNetRefHandleCreationInfo(Context, Handle, InternalObjectIndex);
 	}
 
 	static const UE::Net::Private::FStaticDestructionInfo* GetStaticDestructionInfo(UObjectReplicationBridge* Bridge, UE::Net::FNetRefHandle Handle)
diff --git a/Engine/Source/Runtime/Net/Iris/Public/Iris/ReplicationSystem/ReplicationBridge.h b/Engine/Source/Runtime/Net/Iris/Public/Iris/ReplicationSystem/ReplicationBridge.h
--- a/Engine/Source/Runtime/Net/Iris/Public/Iris/ReplicationSystem/ReplicationBridge.h
+++ b/Engine/Source/Runtime/Net/Iris/Public/Iris/ReplicationSystem/ReplicationBridge.h
@@ -222,7 +222,7 @@ protected:
 	IRISCORE_API void InternalDetachInstanceFromNetRefHandle(UE::Net::Private::FInternalNetRefIndex InternalObjectIndex);
 
 	/** Destroy the handle and all internal book keeping associated with it. */
-	IRISCORE_API void InternalDestroyNetObject(FNetRefHandle Handle);
+	IRISCORE_API void InternalDestroyNetObject(UE::Net::Private::FInternalNetRefIndex InternalObjectIndex);
 	
 	/** Add SubObjectHandle as SubObject to OwnerHandle. */
 	IRISCORE_API void InternalAddSubObject(FNetRefHandle OwnerHandle, FNetRefHandle SubObjectHandle, FNetRefHandle InsertRelativeToSubObjectHandle, UE::Net::ESubObjectInsertionOrder InsertionOrder);
@@ -282,7 +282,7 @@ private:
 
 	// Adds the Handle to the list of handles pending deferred EndReplication, if bIsImmediate is true the object will be destroyed after the next update, otherwise
 	// it will be kept around until the handle is no longer ref-counted by any connection. It will however be removed from the set of scopeable objects after the first update so new connections will not add it to their scope.
-	void AddPendingEndReplication(FNetRefHandle Handle, EEndReplicationFlags DestroyFlags, EPendingEndReplicationImmediate Immediate = EPendingEndReplicationImmediate::No);
+	void AddPendingEndReplication(FNetRefHandle Handle, UE::Net::Private::FInternalNetRefIndex InternalObjectIndex, EEndReplicationFlags DestroyFlags, EPendingEndReplicationImmediate Immediate = EPendingEndReplicationImmediate::No);
 
 	void CallPreSendUpdate(float DeltaSeconds);
 	void CallPreSendUpdateSingleHandle(FNetRefHandle Handle);
@@ -296,18 +296,15 @@ private:
 
 private:
 
-	void InternalFlushStateData(UE::Net::FNetSerializationContext& SerializationContext, UE::Net::Private::FChangeMaskCache& ChangeMaskCache, UE::Net::FNetBitStreamWriter& ChangeMaskWriter, uint32 InternalObjectIndex);
+	void InternalFlushStateData(UE::Net::FNetSerializationContext& SerializationContext, UE::Net::Private::FChangeMaskCache& ChangeMaskCache, UE::Net::FNetBitStreamWriter& ChangeMaskWriter, UE::Net::Private::FInternalNetRefIndex InternalObjectIndex);
 	// Internal method to copy state data for Handle
-	void InternalFlushStateData(FNetRefHandle Handle);
+	void InternalFlushStateData(UE::Net::Private::FInternalNetRefIndex InternalObjectIndex);
 
 	// Internal method to copy state data for Handle and any SubObjects and mark them as being torn-off
 	void InternalTearOff(FNetRefHandle OwnerHandle);
 
-	// Destroy all SubObjects owned by provided handle
-	void InternalDestroySubObjects(FNetRefHandle OwnerHandle, EEndReplicationFlags Flags);
-
-	// Internal method to stop replication by filtering out the replicated object. Then once it is no longer replicating to any client it will be destroyed.
-	void InternalRequestAsyncStopReplication(FNetRefHandle Handle, EEndReplicationFlags Flags);
+	// Internal method to stop replication by filtering out the replicated object. Once it is no longer replicating to any client it will be destroyed.
+	void InternalRequestAsyncStopReplication(UE::Net::Private::FInternalNetRefIndex InternalObjectIndex, EEndReplicationFlags Flags);
 
 	/**
 	 * Called from ReplicationSystem when a streaming container is about to unload.
@@ -315,10 +312,20 @@ private:
 	 */
 	void NotifyStreamingContainerUnload(const UObject* Container);
 
+	void RemoveFromAllGroups(UE::Net::Private::FInternalNetRefIndex InternalObjectIndex);
+
+	// Destroy is split into two phases in order to support replicating data for objects where the external representation might have been destroyed.
+	void InternalStartDestroyLocalSubObjects(UE::Net::Private::FInternalNetRefIndex InternalIndex, EEndReplicationFlags Flags);
+	void InternalFinishDestroyLocalSubObjects(UE::Net::Private::FInternalNetRefIndex InternalIndex, EEndReplicationFlags Flags);
+	void InternalStartDestroyLocalObject(UE::Net::Private::FInternalNetRefIndex InternalIndex, EEndReplicationFlags Flags);
+	void InternalFinishDestroyLocalObject(UE::Net::Private::FInternalNetRefIndex InternalIndex, EEndReplicationFlags Flags);
+
+	// Destroy local NetObject associated with handle
 	void DestroyLocalNetHandle(FNetRefHandle Handle, EEndReplicationFlags Flags);
 
-	// Tear-off all handles in the PendingTearOff list that has not yet been torn-off
-	void TearOffHandlesPendingTearOff();
+	// Update handles pending end replication before we try to anything with them
+	// Detach invalid objects, tear-off handle pending tear off that has not yet been torn-off
+	void PreUpdateHandlesPendingEndReplication();
 
 	// Update all the handles pending EndReplication
 	void UpdateHandlesPendingEndReplication();
@@ -354,9 +361,16 @@ private:
 
 	struct FPendingEndReplicationInfo
 	{
-		FPendingEndReplicationInfo(FNetRefHandle InHandle, EEndReplicationFlags InEndReplicationFlags, EPendingEndReplicationImmediate InImmediate) : Handle(InHandle), EndReplicationFlags(InEndReplicationFlags), Immediate(InImmediate) {}
+		FPendingEndReplicationInfo(FNetRefHandle InHandle, UE::Net::Private::FInternalNetRefIndex InObjectIndex, EEndReplicationFlags InEndReplicationFlags, EPendingEndReplicationImmediate InImmediate)
+		: Handle(InHandle)
+		, ObjectIndex(InObjectIndex)
+		, EndReplicationFlags(InEndReplicationFlags)
+		, Immediate(InImmediate)
+		{
+		}
 
 		FNetRefHandle Handle;
+		UE::Net::Private::FInternalNetRefIndex ObjectIndex;
 		EEndReplicationFlags EndReplicationFlags;
 		EPendingEndReplicationImmediate Immediate;
 	};
