Iris: rename flags related to stopping replication that use the term "Migrated" to better reflect their actual functionality. These flags now use "ProxyReuse" since they're only supported by clients with replication systems configured to use UE::Net::EProxyType::Backend, and are meant to allow the client to reuse the object later.
Also remove EDetachReason::Migrated since it's not intended to be exposed to factory implementations.

#Synced-CL 49887316

diff --git a/Engine/Plugins/Runtime/ReplicationSystemTestPlugin/Source/Private/Tests/ReplicationSystem/TestReplicationSystemProxy.cpp b/Engine/Plugins/Runtime/ReplicationSystemTestPlugin/Source/Private/Tests/ReplicationSystem/TestReplicationSystemProxy.cpp
--- a/Engine/Plugins/Runtime/ReplicationSystemTestPlugin/Source/Private/Tests/ReplicationSystem/TestReplicationSystemProxy.cpp
+++ b/Engine/Plugins/Runtime/ReplicationSystemTestPlugin/Source/Private/Tests/ReplicationSystem/TestReplicationSystemProxy.cpp
@@ -207,8 +207,8 @@ UE_NET_TEST_FIXTURE(FReplicationSystemProxyTestFixture, ReuseObjectPropertyRepli
 	UE_NET_ASSERT_EQ(Server1RootObject->NetRefHandle.GetId(), Server2RootObject->NetRefHandle.GetId());
 	UE_NET_ASSERT_EQ(Server1SubObject->NetRefHandle.GetId(), Server2SubObject->NetRefHandle.GetId());
 
-	Server1->DestroyObject(Server1SubObject, EEndReplicationFlags::Migrated);
-	Server1->DestroyObject(Server1RootObject, EEndReplicationFlags::Migrated);
+	Server1->DestroyObject(Server1SubObject, EEndReplicationFlags::ProxyReuse);
+	Server1->DestroyObject(Server1RootObject, EEndReplicationFlags::ProxyReuse);
 	Server1RootObject = nullptr;
 	Server1SubObject = nullptr;
 
@@ -325,8 +325,10 @@ UE_NET_TEST_FIXTURE(FReplicationSystemProxyTestFixture, ReuseObjectDestruction)
 	UE_NET_ASSERT_EQ(Server1RootObject->NetRefHandle.GetId(), Server2RootObject->NetRefHandle.GetId());
 	UE_NET_ASSERT_EQ(Server1SubObject->NetRefHandle.GetId(), Server2SubObject->NetRefHandle.GetId());
 
-	Server1->DestroyObject(Server1RootObject, EEndReplicationFlags::Migrated);
+	Server1->DestroyObject(Server1SubObject, EEndReplicationFlags::ProxyReuse);
+	Server1->DestroyObject(Server1RootObject, EEndReplicationFlags::ProxyReuse);
 	Server1RootObject = nullptr;
+	Server1SubObject = nullptr;
 
 	// Send and deliver packet, servers -> proxy
 	UpdateAndSendServerToAllProxies(Server1);
diff --git a/Engine/Source/Runtime/Engine/Private/Actor.cpp b/Engine/Source/Runtime/Engine/Private/Actor.cpp
--- a/Engine/Source/Runtime/Engine/Private/Actor.cpp
+++ b/Engine/Source/Runtime/Engine/Private/Actor.cpp
@@ -1360,7 +1360,7 @@ void AActor::PostMigrate(const FUObjectMigrationContext& MigrationContext)
 		if (MigrationSendType == EObjectMigrationSendType::ReassignOwnership)
 		{
 			FStopReplicatingActorParams StopRepParams(EEndPlayReason::Destroyed);
-			StopRepParams.StopReplicatingReason = UE::Net::EStopReplicatingReason::Migrated;
+			StopRepParams.StopReplicatingReason = UE::Net::EStopReplicatingReason::ProxyReuse;
 			UE::Net::FDeferredReplicationSystemCalls::StopReplicatingActorDefer(this, StopRepParams);
 		}
 	}
diff --git a/Engine/Source/Runtime/Engine/Private/Net/Iris/ReplicationSystem/EngineReplicationBridge.cpp b/Engine/Source/Runtime/Engine/Private/Net/Iris/ReplicationSystem/EngineReplicationBridge.cpp
--- a/Engine/Source/Runtime/Engine/Private/Net/Iris/ReplicationSystem/EngineReplicationBridge.cpp
+++ b/Engine/Source/Runtime/Engine/Private/Net/Iris/ReplicationSystem/EngineReplicationBridge.cpp
@@ -651,14 +651,14 @@ EEndReplicationFlags UEngineReplicationBridge::ConvertToEndReplicationFlags(UE::
 	using namespace UE::Net;
 
 	EEndReplicationFlags EndReplicationFlags = EEndReplicationFlags::None;
-	if (StopReplicatingReason == EStopReplicatingReason::Destroyed || StopReplicatingReason == EStopReplicatingReason::Migrated)
+	if (StopReplicatingReason == EStopReplicatingReason::Destroyed || StopReplicatingReason == EStopReplicatingReason::ProxyReuse)
 	{
 		EndReplicationFlags |= EEndReplicationFlags::Destroy | EEndReplicationFlags::DestroyNetHandle | EEndReplicationFlags::ClearNetPushId;
 	}
 	
-	if (StopReplicatingReason == EStopReplicatingReason::Migrated)
+	if (StopReplicatingReason == EStopReplicatingReason::ProxyReuse)
 	{
-		EndReplicationFlags |= EEndReplicationFlags::Migrated;
+		EndReplicationFlags |= EEndReplicationFlags::ProxyReuse;
 	}
 
 	// If we are shutting down we do not need to validate that we are not detaching remote instances by accident.
diff --git a/Engine/Source/Runtime/Engine/Private/PlayerController.cpp b/Engine/Source/Runtime/Engine/Private/PlayerController.cpp
--- a/Engine/Source/Runtime/Engine/Private/PlayerController.cpp
+++ b/Engine/Source/Runtime/Engine/Private/PlayerController.cpp
@@ -5082,7 +5082,7 @@ void APlayerController::PostMigrate(const FUObjectMigrationContext& MigrationCon
 		if (NetDriver->GetReplicationSystem())
 		{
 			FStopReplicatingActorParams StopRepParams(EEndPlayReason::Destroyed);
-			StopRepParams.StopReplicatingReason = UE::Net::EStopReplicatingReason::Migrated;
+			StopRepParams.StopReplicatingReason = UE::Net::EStopReplicatingReason::ProxyReuse;
 			UE::Net::FDeferredReplicationSystemCalls::StopReplicatingActorDefer(Actor, StopRepParams);
 		}
 		else
diff --git a/Engine/Source/Runtime/Engine/Public/Net/Iris/ReplicationSystem/EngineReplicationBridge.h b/Engine/Source/Runtime/Engine/Public/Net/Iris/ReplicationSystem/EngineReplicationBridge.h
--- a/Engine/Source/Runtime/Engine/Public/Net/Iris/ReplicationSystem/EngineReplicationBridge.h
+++ b/Engine/Source/Runtime/Engine/Public/Net/Iris/ReplicationSystem/EngineReplicationBridge.h
@@ -40,8 +40,8 @@ enum class EStopReplicatingReason : uint32
 	RemovedFromWorld,
 	/** When the application is being exited. */
 	Quit,
-	/** When an actor is moved from one game server to another. */
-	Migrated,
+	/** An object is removed from a server but the client should keep it and maintain its internal data, expecting to receive object updates from another server later. Only supported by clients with replication systems configured to use UE::Net::EProxyType::Backend. */
+	ProxyReuse,
 };
 
 ENGINE_API const TCHAR* LexToString(EStopReplicatingReason Reason);
diff --git a/Engine/Source/Runtime/Net/Iris/Private/Iris/ReplicationSystem/NetObjectFactory.cpp b/Engine/Source/Runtime/Net/Iris/Private/Iris/ReplicationSystem/NetObjectFactory.cpp
--- a/Engine/Source/Runtime/Net/Iris/Private/Iris/ReplicationSystem/NetObjectFactory.cpp
+++ b/Engine/Source/Runtime/Net/Iris/Private/Iris/ReplicationSystem/NetObjectFactory.cpp
@@ -190,11 +190,6 @@ const TCHAR* LexToString(UE::Net::EDetachReason Reason)
 			return TEXT("TornOff");
 		} break;
 
-		case EDetachReason::Migrated:
-		{
-			return TEXT("Migrated");
-		} break;
-
 		case EDetachReason::RootObjectStillRelevant:
 		{
 			return TEXT("RootObjectStillRelevant");
diff --git a/Engine/Source/Runtime/Net/Iris/Private/Iris/ReplicationSystem/NetRefHandleManager.cpp b/Engine/Source/Runtime/Net/Iris/Private/Iris/ReplicationSystem/NetRefHandleManager.cpp
--- a/Engine/Source/Runtime/Net/Iris/Private/Iris/ReplicationSystem/NetRefHandleManager.cpp
+++ b/Engine/Source/Runtime/Net/Iris/Private/Iris/ReplicationSystem/NetRefHandleManager.cpp
@@ -52,8 +52,8 @@ const TCHAR* LexToString(EInternalDetachReason DetachReason)
 			return TEXT("NotRelevant");
 		case EInternalDetachReason::StaticDestroyed:
 			return TEXT("StaticDestroyed");
-		case EInternalDetachReason::Migrated:
-			return TEXT("Migrated");
+		case EInternalDetachReason::ProxyReuse:
+			return TEXT("ProxyReuse");
 		case EInternalDetachReason::TornOff:
 			return TEXT("TornOff");
 		case EInternalDetachReason::Max:
diff --git a/Engine/Source/Runtime/Net/Iris/Private/Iris/ReplicationSystem/NetRefHandleManager.h b/Engine/Source/Runtime/Net/Iris/Private/Iris/ReplicationSystem/NetRefHandleManager.h
--- a/Engine/Source/Runtime/Net/Iris/Private/Iris/ReplicationSystem/NetRefHandleManager.h
+++ b/Engine/Source/Runtime/Net/Iris/Private/Iris/ReplicationSystem/NetRefHandleManager.h
@@ -70,8 +70,8 @@ ENUM_CLASS_FLAGS(ERemoveDependentObjectFlags);
 enum class ECloseReplicationReason : uint8
 {
 	None = 0,
-	/** When an actor is migrated from one game server to another. */
-	Migrated = 1U,
+	/** An object is removed from a server but the client should keep it and maintain its internal data, expecting to receive object updates from another server later. Only supported by clients with replication systems configured to use UE::Net::EProxyType::Backend. */
+	ProxyReuse = 1U,
 	/** The number of bits required to hold the largest value above. */
 	MaxBits = 1U
 };
diff --git a/Engine/Source/Runtime/Net/Iris/Private/Iris/ReplicationSystem/NetRefHandleManagerTypes.h b/Engine/Source/Runtime/Net/Iris/Private/Iris/ReplicationSystem/NetRefHandleManagerTypes.h
--- a/Engine/Source/Runtime/Net/Iris/Private/Iris/ReplicationSystem/NetRefHandleManagerTypes.h
+++ b/Engine/Source/Runtime/Net/Iris/Private/Iris/ReplicationSystem/NetRefHandleManagerTypes.h
@@ -30,8 +30,8 @@ enum class EInternalDetachReason : uint8
 	NotRelevant,
 	/** When a stable replicated object is destroyed even when the client loads it locally. */
 	StaticDestroyed,
-	/** When a replicated object is migrated from one game server to another. */
-	Migrated,
+	/** An object is removed from a server but the client should keep it and maintain its internal data, expecting to receive object updates from another server later. Only supported by clients with replication systems configured to use UE::Net::EProxyType::Backend. */
+	ProxyReuse,
 	/** When a replicated object is torn off by the authority. */
 	TornOff,
 	/** Add new entries above this one. */
diff --git a/Engine/Source/Runtime/Net/Iris/Private/Iris/ReplicationSystem/ObjectReplicationBridge.cpp b/Engine/Source/Runtime/Net/Iris/Private/Iris/ReplicationSystem/ObjectReplicationBridge.cpp
--- a/Engine/Source/Runtime/Net/Iris/Private/Iris/ReplicationSystem/ObjectReplicationBridge.cpp
+++ b/Engine/Source/Runtime/Net/Iris/Private/Iris/ReplicationSystem/ObjectReplicationBridge.cpp
@@ -985,9 +985,10 @@ UE::Net::EDetachReason ConvertDetachReasonFromInternalToPublic(UE::Net::Private:
 			return EDetachReason::StaticDestroyed;
 		} break;
 
-		case EInternalDetachReason::Migrated:
+		case EInternalDetachReason::ProxyReuse:
 		{
-			return EDetachReason::Migrated;
+			// ProxyReuse isn't currently exposed to factories
+			ensureMsgf(false, TEXT("ProxyReuse isn't a public DetachReason"));
 		} break;
 
 		case EInternalDetachReason::TornOff:
@@ -1061,41 +1062,45 @@ void UObjectReplicationBridge::DestroyNetObjectFromRemote(UE::Net::Private::FInt
 		}
 	}
 
-	// Convert from internal to public reason
-	EDetachReason PublicDetachReason = ConvertDetachReasonFromInternalToPublic(InternalDetachReason);
-
-	UE_LOG_OBJECTREPLICATIONBRIDGE(Verbose, TEXT("DestroyNetObjectFromRemote for %s | InternalReason: %s | PublicReason: %s |  DestroyReason: %s | DestroyFlags: %s "), *NetRefHandleManager->PrintObjectFromIndex(InternalObjectIndex), LexToString(InternalDetachReason), LexToString(PublicDetachReason), LexToString(DestroyReason), LexToString(DestroyFlags) );
-
-	const FNetRefHandle DestroyedObjectHandle = ObjectData.RefHandle;
-	ensureMsgf(DestroyedObjectHandle.IsValid(), TEXT("Object with InternalIndex assigned an invalid handle: %s"), *NetRefHandleManager->PrintObjectFromIndex(InternalObjectIndex));
-	if (!DestroyedObjectHandle.IsValid())
+	// Reused objects need to keep the Handle->Pointer binding alive since proxies reuse pointers instead of recreating them. This won't be required in the future if proxies simply do not instantiate a local instance and use protocol data directly.
+	if (InternalDetachReason != EInternalDetachReason::ProxyReuse)
 	{
-		return;
-	}
+		// Convert from internal to public reason
+		EDetachReason PublicDetachReason = ConvertDetachReasonFromInternalToPublic(InternalDetachReason);
 
-	if (UE::Net::IsUsingNewDetachAPI() == false)
-	{
-		// if a root object is detaching, we want to detach all subobjects first to ensure we execute expected callbacks
-		// We keep tracking them internally
-		if (ObjectData.IsRootObject())
+		UE_LOG_OBJECTREPLICATIONBRIDGE(Verbose, TEXT("DestroyNetObjectFromRemote for %s | InternalReason: %s | PublicReason: %s |  DestroyReason: %s | DestroyFlags: %s "), *NetRefHandleManager->PrintObjectFromIndex(InternalObjectIndex), LexToString(InternalDetachReason), LexToString(PublicDetachReason), LexToString(DestroyReason), LexToString(DestroyFlags));
+
+		const FNetRefHandle DestroyedObjectHandle = ObjectData.RefHandle;
+		ensureMsgf(DestroyedObjectHandle.IsValid(), TEXT("Object with InternalIndex assigned an invalid handle: %s"), *NetRefHandleManager->PrintObjectFromIndex(InternalObjectIndex));
+		if (!DestroyedObjectHandle.IsValid())
 		{
-			DetachSubObjectInstancesFromRemote(InternalObjectIndex, DestroyReason, DestroyFlags);
+			return;
 		}
-		
-		// Tell factories to take care of the detached object.
-		DetachInstanceFromRemote(DestroyedObjectHandle, DestroyReason, DestroyFlags, ObjectData.NetFactoryId);
-	}
-	else
-	{
-		if (ObjectData.IsRootObject())
+
+		if (UE::Net::IsUsingNewDetachAPI() == false)
 		{
-			//TODO: Test how Iris behaves if a subobject is destroyed with a different reason (tearoff, etc.) from it's root object and we receive both the destruction messages at the same or different time.
-			//It feels like the root object reason migh always have priority over everything else ?
-			DetachRootObjectFromRemote(DestroyedObjectHandle, PublicDetachReason, ObjectData.NetFactoryId);
+			// if a root object is detaching, we want to detach all subobjects first to ensure we execute expected callbacks
+			// We keep tracking them internally
+			if (ObjectData.IsRootObject())
+			{
+				DetachSubObjectInstancesFromRemote(InternalObjectIndex, DestroyReason, DestroyFlags);
+			}
+
+			// Tell factories to take care of the detached object.
+			DetachInstanceFromRemote(DestroyedObjectHandle, DestroyReason, DestroyFlags, ObjectData.NetFactoryId);
 		}
 		else
 		{
-			DetachSubObjectFromRemote(DestroyedObjectHandle, PublicDetachReason, ObjectData.NetFactoryId);
+			if (ObjectData.IsRootObject())
+			{
+				//TODO: Test how Iris behaves if a subobject is destroyed with a different reason (tearoff, etc.) from it's root object and we receive both the destruction messages at the same or different time.
+				//It feels like the root object reason migh always have priority over everything else ?
+				DetachRootObjectFromRemote(DestroyedObjectHandle, PublicDetachReason, ObjectData.NetFactoryId);
+			}
+			else
+			{
+				DetachSubObjectFromRemote(DestroyedObjectHandle, PublicDetachReason, ObjectData.NetFactoryId);
+			}
 		}
 	}
 
@@ -1404,7 +1409,7 @@ void UObjectReplicationBridge::DetachRootObjectFromRemote(FNetRefHandle RootObje
 			// The deprecated functions would only be called if these conditions are met.
 			const bool bCallDeprecatedFunctions = SubObjectDestroyContext.DestroyedInstance != nullptr &&
 											SubObjectDestroyContext.RootObject != nullptr &&
-											DestroyReason != EReplicationBridgeDestroyInstanceReason::Migrated;
+											DestroyReason != EReplicationBridgeDestroyInstanceReason::ProxyReuse;
 			if (bCallDeprecatedFunctions)
 			{
 				RootObjectFactory->SubObjectDetachedFromReplication(SubObjectDestroyContext);
@@ -1420,26 +1425,23 @@ void UObjectReplicationBridge::DetachRootObjectFromRemote(FNetRefHandle RootObje
 			.Reason = DetachReason,
 		};
 
-		if (DetachReason != EDetachReason::Migrated)
-		{
-			// Tell the root object he will lose this subobject
-			RootObjectFactory->SubObjectDetachedFromReplication(DetachContext, SubObjectDetachContext);
+		// Tell the root object he will lose this subobject
+		RootObjectFactory->SubObjectDetachedFromReplication(DetachContext, SubObjectDetachContext);
 
-			// Now ask the subobject's factory to handle the detachment
-			if (UNetObjectFactory* SubObjectFactory = GetNetFactory(SubObjectData.NetFactoryId))
+		// Now ask the subobject's factory to handle the detachment
+		if (UNetObjectFactory* SubObjectFactory = GetNetFactory(SubObjectData.NetFactoryId))
+		{
+			/*PRAGMA_DISABLE_DEPRECATION_WARNINGS
+			if (bCallDeprecatedFunctions)
 			{
-				/*PRAGMA_DISABLE_DEPRECATION_WARNINGS
-				if (bCallDeprecatedFunctions)
-				{
-					SubObjectFactory->DetachedFromReplication(SubObjectDestroyContext);
-				}
-				PRAGMA_ENABLE_DEPRECATION_WARNINGS*/
-
-				SubObjectFactory->DetachedFromReplication(DetachContext, SubObjectDetachContext);
+				SubObjectFactory->DetachedFromReplication(SubObjectDestroyContext);
 			}
+			PRAGMA_ENABLE_DEPRECATION_WARNINGS*/
 
-			UnregisterInstance(SubObjectHandle);
+			SubObjectFactory->DetachedFromReplication(DetachContext, SubObjectDetachContext);
 		}
+
+		UnregisterInstance(SubObjectHandle);
 	}
 
 	// Now that subobjects have been detached, handle the root object itself
@@ -1456,7 +1458,7 @@ void UObjectReplicationBridge::DetachRootObjectFromRemote(FNetRefHandle RootObje
 			.DestroyFlags = DestroyFlags,
 		};
 		if (DestroyContext.DestroyedInstance != nullptr &&
-			DestroyReason != EReplicationBridgeDestroyInstanceReason::Migrated)
+			DestroyReason != EReplicationBridgeDestroyInstanceReason::ProxyReuse)
 		{
 			RootObjectFactory->DetachedFromReplication(DestroyContext);
 		}
@@ -1471,12 +1473,6 @@ void UObjectReplicationBridge::DetachRootObjectFromRemote(FNetRefHandle RootObje
 		.Reason = DetachReason,
 	};
 
-	// Migrated objects need to keep the Handle->Pointer binding alive since proxies reuse pointers instead of recreating them. This won't be required in the future if proxies simply do not instantiate a local instance and use protocol data directly.
-	if (DetachReason == EDetachReason::Migrated)
-	{
-		return;
-	}
-
 	RootObjectFactory->DetachedFromReplication(DetachContext, {});
 
 	UnregisterInstance(RootObjectHandle);
@@ -1559,7 +1555,7 @@ void UObjectReplicationBridge::DetachSubObjectFromRemote(FNetRefHandle DetachedS
 			const bool bCallDeprecatedFunctions =
 				SubObjectDestroyContext.DestroyedInstance != nullptr &&
 				SubObjectDestroyContext.RootObject != nullptr &&
-				DestroyReason != EReplicationBridgeDestroyInstanceReason::Migrated;
+				DestroyReason != EReplicationBridgeDestroyInstanceReason::ProxyReuse;
 			if (bCallDeprecatedFunctions)
 			{
 				RootObjectFactory->SubObjectDetachedFromReplication(SubObjectDestroyContext);
@@ -1571,12 +1567,6 @@ void UObjectReplicationBridge::DetachSubObjectFromRemote(FNetRefHandle DetachedS
 		RootObjectFactory->SubObjectDetachedFromReplication(DetachContext, SubObjectDetachContext);
 	}
 
-	// Migrated objects need to keep the Handle->Pointer binding alive since proxies reuse pointers instead of recreating them. This won't be required in the future if proxies simply do not instantiate a local instance and use protocol data directly.
-	if (DetachReason == EDetachReason::Migrated)
-	{
-		return;
-	}
-
 	SubObjectFactory->DetachedFromReplication(DetachContext, SubObjectDetachContext);
 
 	UnregisterInstance(DetachedSubObjectHandle);
@@ -1625,7 +1615,7 @@ PRAGMA_DISABLE_DEPRECATION_WARNINGS
 		// Now ask the object's factory to handle the detachment
 		if (ensureMsgf(NetFactoryId != InvalidNetObjectFactoryId, TEXT("Received invalid netfactory for %s: %s"), *Handle.ToString(), *GetPathNameSafe(DestroyContext.DestroyedInstance)))
 		{
-			if (DestroyReason != EReplicationBridgeDestroyInstanceReason::Migrated)
+			if (DestroyReason != EReplicationBridgeDestroyInstanceReason::ProxyReuse)
 			{
 				if (UNetObjectFactory* Factory = GetNetFactory(NetFactoryId))
 				{
@@ -1636,7 +1626,7 @@ PRAGMA_DISABLE_DEPRECATION_WARNINGS
 	}
 PRAGMA_ENABLE_DEPRECATION_WARNINGS
 
-	if (DestroyReason != EReplicationBridgeDestroyInstanceReason::Migrated)
+	if (DestroyReason != EReplicationBridgeDestroyInstanceReason::ProxyReuse)
 	{
 		UnregisterInstance(Handle);
 	}
diff --git a/Engine/Source/Runtime/Net/Iris/Private/Iris/ReplicationSystem/ReplicationBridge.cpp b/Engine/Source/Runtime/Net/Iris/Private/Iris/ReplicationSystem/ReplicationBridge.cpp
--- a/Engine/Source/Runtime/Net/Iris/Private/Iris/ReplicationSystem/ReplicationBridge.cpp
+++ b/Engine/Source/Runtime/Net/Iris/Private/Iris/ReplicationSystem/ReplicationBridge.cpp
@@ -459,10 +459,10 @@ void UReplicationBridge::StopReplicatingNetRefHandle(FNetRefHandle Handle, EEndR
 		// Store the stop replication reason
 		ObjectData.InternalDetachReason = EInternalDetachReason::Normal;
 
-		if (EnumHasAnyFlags(EndReplicationFlags, EEndReplicationFlags::Migrated))
+		if (EnumHasAnyFlags(EndReplicationFlags, EEndReplicationFlags::ProxyReuse))
 		{
-			ObjectData.CloseReason = ECloseReplicationReason::Migrated;
-			ObjectData.InternalDetachReason = EInternalDetachReason::Migrated;
+			ObjectData.CloseReason = ECloseReplicationReason::ProxyReuse;
+			ObjectData.InternalDetachReason = EInternalDetachReason::ProxyReuse;
 		}
 
 		if (EnumHasAnyFlags(EndReplicationFlags, EEndReplicationFlags::TearOff))
diff --git a/Engine/Source/Runtime/Net/Iris/Private/Iris/ReplicationSystem/ReplicationBridgeTypes.cpp b/Engine/Source/Runtime/Net/Iris/Private/Iris/ReplicationSystem/ReplicationBridgeTypes.cpp
--- a/Engine/Source/Runtime/Net/Iris/Private/Iris/ReplicationSystem/ReplicationBridgeTypes.cpp
+++ b/Engine/Source/Runtime/Net/Iris/Private/Iris/ReplicationSystem/ReplicationBridgeTypes.cpp
@@ -42,10 +42,10 @@ FString LexToString(EEndReplicationFlags EndReplicationFlags)
 		if (!Flags.IsEmpty()) { Flags += TEXT(','); };
 		Flags += TEXT("SkipPendingEndReplicationValidation");
 	}
-	if (EnumHasAnyFlags(EndReplicationFlags, EEndReplicationFlags::Migrated))
+	if (EnumHasAnyFlags(EndReplicationFlags, EEndReplicationFlags::ProxyReuse))
 	{
 		if (!Flags.IsEmpty()) { Flags += TEXT(','); };
-		Flags += TEXT("Migrated");
+		Flags += TEXT("ProxyReuse");
 	}
 	if (EnumHasAnyFlags(EndReplicationFlags, EEndReplicationFlags::AllowAsyncStopReplication))
 	{
@@ -72,9 +72,9 @@ const TCHAR* LexToString(EReplicationBridgeDestroyInstanceReason Reason)
 	{
 		return TEXT("Destroy");
 	}
-	case EReplicationBridgeDestroyInstanceReason::Migrated:
+	case EReplicationBridgeDestroyInstanceReason::ProxyReuse:
 	{
-		return TEXT("Migrated");
+		return TEXT("ProxyReuse");
 	}
 	default:
 	{
diff --git a/Engine/Source/Runtime/Net/Iris/Private/Iris/ReplicationSystem/ReplicationReader.cpp b/Engine/Source/Runtime/Net/Iris/Private/Iris/ReplicationSystem/ReplicationReader.cpp
--- a/Engine/Source/Runtime/Net/Iris/Private/Iris/ReplicationSystem/ReplicationReader.cpp
+++ b/Engine/Source/Runtime/Net/Iris/Private/Iris/ReplicationSystem/ReplicationReader.cpp
@@ -631,9 +631,9 @@ void FReplicationReader::EndReplication(FInternalNetRefIndex InternalIndex, bool
 	{
 		DestroyReason = EReplicationBridgeDestroyInstanceReason::TearOff;
 	}
-	else if (CloseReason == ECloseReplicationReason::Migrated)
+	else if (CloseReason == ECloseReplicationReason::ProxyReuse)
 	{
-		DestroyReason = EReplicationBridgeDestroyInstanceReason::Migrated;
+		DestroyReason = EReplicationBridgeDestroyInstanceReason::ProxyReuse;
 	}
 	else if (bDestroyInstance)
 	{
diff --git a/Engine/Source/Runtime/Net/Iris/Public/Iris/ReplicationSystem/NetObjectFactory.h b/Engine/Source/Runtime/Net/Iris/Public/Iris/ReplicationSystem/NetObjectFactory.h
--- a/Engine/Source/Runtime/Net/Iris/Public/Iris/ReplicationSystem/NetObjectFactory.h
+++ b/Engine/Source/Runtime/Net/Iris/Public/Iris/ReplicationSystem/NetObjectFactory.h
@@ -39,9 +39,6 @@ enum class EDetachReason : uint32
 	/** Object is being torn off permanently */
 	TornOff,
 
-	/** Object is migrating */
-	Migrated,
-
 	/** Special entry used in FSubObjectDetachContext when the root object stays relevant. */
 	RootObjectStillRelevant,
 
diff --git a/Engine/Source/Runtime/Net/Iris/Public/Iris/ReplicationSystem/ReplicationBridgeTypes.h b/Engine/Source/Runtime/Net/Iris/Public/Iris/ReplicationSystem/ReplicationBridgeTypes.h
--- a/Engine/Source/Runtime/Net/Iris/Public/Iris/ReplicationSystem/ReplicationBridgeTypes.h
+++ b/Engine/Source/Runtime/Net/Iris/Public/Iris/ReplicationSystem/ReplicationBridgeTypes.h
@@ -24,10 +24,10 @@ enum class EEndReplicationFlags : uint32
 	ClearNetPushId						= DestroyNetHandle << 1U,
 	/** Skip bPendingEndReplication Validation, In some cases we want to allow detaching instance from replicated object on clients, such as when shutting down */
 	SkipPendingEndReplicationValidation = ClearNetPushId << 1U,
-	/** An actor is migrated between servers. Only supported by clients with replication systems configured to use UE::Net::EProxyType::Backend. */
-	Migrated							= SkipPendingEndReplicationValidation << 1U,
+	/** An object is removed from a server but the client should keep it and maintain its internal data, expecting to receive object updates from another server later. Only supported by clients with replication systems configured to use UE::Net::EProxyType::Backend. */
+	ProxyReuse							= SkipPendingEndReplicationValidation << 1U,
 	/** If true, replication can be resumed effectively canceling the EndReplication by issuing a StartReplication before the StopReplication has been acknowledged by all clients. */
-	AllowAsyncStopReplication = Migrated << 1U,
+	AllowAsyncStopReplication			= ProxyReuse << 1U,
 };
 ENUM_CLASS_FLAGS(EEndReplicationFlags);
 
@@ -63,7 +63,7 @@ enum class EReplicationBridgeDestroyInstanceReason : uint32
 	DoNotDestroy,
 	TearOff,
 	Destroy,
-	Migrated,
+	ProxyReuse,
 };
 IRISCORE_API const TCHAR* LexToString(EReplicationBridgeDestroyInstanceReason Reason);
 
