Iris
- Avoid trigger users/error warning/ensure when pruning async stop replication objects that are being invalidated.
- Allow proper cleanup of bound instances that do not properly call StopReplication during pruning

#Synced-CL 49833605

diff --git a/Engine/Source/Runtime/Net/Iris/Private/Iris/ReplicationSystem/ObjectReplicationBridge.cpp b/Engine/Source/Runtime/Net/Iris/Private/Iris/ReplicationSystem/ObjectReplicationBridge.cpp
--- a/Engine/Source/Runtime/Net/Iris/Private/Iris/ReplicationSystem/ObjectReplicationBridge.cpp
+++ b/Engine/Source/Runtime/Net/Iris/Private/Iris/ReplicationSystem/ObjectReplicationBridge.cpp
@@ -2025,7 +2025,7 @@ void UObjectReplicationBridge::PruneStaleObjects()
 	TArray<FNetRefHandle> StaleObjects;
 
 	// Detect stale references and try to kill/report them
-	auto DetectStaleObjectsFunc = [&LocalNetRefHandleManager, &StaleObjects, &ReplicatedInstances](uint32 InternalNetHandleIndex)
+	auto DetectStaleObjectsFunc = [&LocalNetRefHandleManager, &StaleObjects, &ReplicatedInstances, this](uint32 InternalNetHandleIndex)
 	{
 		if (!IsValid(ReplicatedInstances[InternalNetHandleIndex]))
 		{
@@ -2034,13 +2034,22 @@ void UObjectReplicationBridge::PruneStaleObjects()
 			{
 				const FReplicationProtocol* Protocol = ObjectData.Protocol;
 				const FNetDebugName* DebugName = Protocol ? Protocol->DebugName : nullptr;
-				UE_LOG(LogIris, Warning, TEXT("UObjectReplicationBridge::PruneStaleObjects ObjectInstance replicated as: %s of Type named:%s has been destroyed without notifying the ReplicationSystem %u"), *ObjectData.RefHandle.ToString(), ToCStr(DebugName), ObjectData.RefHandle.GetReplicationSystemId());
 
-				// If the instance protocol is bound, then this is an error and we cannot safely cleanup as unbinding abound instance protocol will modify bound states
-				if (ensureMsgf(!EnumHasAnyFlags(ObjectData.InstanceProtocol->InstanceTraits, EReplicationInstanceProtocolTraits::IsBound), TEXT("UObjectReplicationBridge::PruneStaleObjects Bound ObjectInstance replicated as: %s has been destroyed without notifying the ReplicationSystem."), *ObjectData.RefHandle.ToString()))
+				// Only log warning if pruned object is not marked as async stopping.
+				UE_CLOG(!ObjectData.bIsAsyncStopping, LogIris, Warning, TEXT("UObjectReplicationBridge::PruneStaleObjects ObjectInstance replicated as: %s of Type named:%s has been destroyed without notifying the ReplicationSystem %u"), *ObjectData.RefHandle.ToString(), ToCStr(DebugName), ObjectData.RefHandle.GetReplicationSystemId());						
+
+				// If instance protocol is bound, we detach and destroy instance protocol without touching the instance
+				if (EnumHasAnyFlags(ObjectData.InstanceProtocol->InstanceTraits, EReplicationInstanceProtocolTraits::IsBound))
 				{
-					StaleObjects.Push(ObjectData.RefHandle);
+					// Ensure, to bring extra attention if we have missed to call StopReplication on bound object.
+					ensureMsgf(ObjectData.bIsAsyncStopping, TEXT("UObjectReplicationBridge::PruneStaleObjects Bound ObjectInstance replicated as: %s has been destroyed without notifying the ReplicationSystem."), *ObjectData.RefHandle.ToString());
+
+					// Destroy instance protocol without modifying it to avoid potentially touching bad objects.
+					const bool bUnbindInstanceProtocol = false;
+					InternalDetachInstanceFromNetRefHandle(InternalNetHandleIndex, bUnbindInstanceProtocol);
 				}
+
+				StaleObjects.Push(ObjectData.RefHandle);
 			}
 		}
 	};
diff --git a/Engine/Source/Runtime/Net/Iris/Private/Iris/ReplicationSystem/ReplicationBridge.cpp b/Engine/Source/Runtime/Net/Iris/Private/Iris/ReplicationSystem/ReplicationBridge.cpp
--- a/Engine/Source/Runtime/Net/Iris/Private/Iris/ReplicationSystem/ReplicationBridge.cpp
+++ b/Engine/Source/Runtime/Net/Iris/Private/Iris/ReplicationSystem/ReplicationBridge.cpp
@@ -233,7 +233,7 @@ void UReplicationBridge::InternalAttachInstanceToNetRefHandle(FNetRefHandle RefH
 	}
 }
 
-void UReplicationBridge::InternalDetachInstanceFromNetRefHandle(UE::Net::Private::FInternalNetRefIndex InternalObjectIndex)
+void UReplicationBridge::InternalDetachInstanceFromNetRefHandle(UE::Net::Private::FInternalNetRefIndex InternalObjectIndex, bool bUnBindInstanceProtocol)
 {
 	using namespace UE::Net;
 	using namespace UE::Net::Private;
@@ -242,7 +242,7 @@ void UReplicationBridge::InternalDetachInstanceFromNetRefHandle(UE::Net::Private
 	{
 		UE_LOG_REPLICATIONBRIDGE(Verbose, TEXT("InternalDetachInstanceToNetHandle Detached: %s"), *NetRefHandleManager->PrintObjectFromIndex(InternalObjectIndex));
 
-		if (EnumHasAnyFlags(InstanceProtocol->InstanceTraits, EReplicationInstanceProtocolTraits::IsBound))
+		if (bUnBindInstanceProtocol && EnumHasAnyFlags(InstanceProtocol->InstanceTraits, EReplicationInstanceProtocolTraits::IsBound))
 		{
 			FReplicationInstanceOperationsInternal::UnbindInstanceProtocol(InstanceProtocol, NetRefHandleManager->GetReplicatedObjectDataNoCheck(InternalObjectIndex).Protocol);
 		}
@@ -321,12 +321,15 @@ void UReplicationBridge::InternalFinishDestroyLocalObject(UE::Net::Private::FInt
 
 void UReplicationBridge::DestroyLocalNetHandle(FNetRefHandle Handle, EEndReplicationFlags EndReplicationFlags)
 {
-	UE_LOG_REPLICATIONBRIDGE(Verbose, TEXT("DestroyLocalNetHandle for %s | EndReplicationFlags: %s"), *NetRefHandleManager->PrintObjectFromNetRefHandle(Handle), *LexToString(EndReplicationFlags));
-
-	const UE::Net::Private::FInternalNetRefIndex InternalReplicationIndex = NetRefHandleManager->GetInternalIndex(Handle);
+	using namespace UE::Net::Private;
 
-	InternalStartDestroyLocalObject(InternalReplicationIndex, EndReplicationFlags);
-	InternalFinishDestroyLocalObject(InternalReplicationIndex, EndReplicationFlags);
+	const FInternalNetRefIndex InternalReplicationIndex = NetRefHandleManager->GetInternalIndex(Handle);
+	if (InternalReplicationIndex != FNetRefHandleManager::InvalidInternalIndex)
+	{
+		UE_LOG_REPLICATIONBRIDGE(Verbose, TEXT("DestroyLocalNetHandle for %s | EndReplicationFlags: %s"), *NetRefHandleManager->PrintObjectFromNetRefHandle(Handle), *LexToString(EndReplicationFlags));
+		InternalStartDestroyLocalObject(InternalReplicationIndex, EndReplicationFlags);
+		InternalFinishDestroyLocalObject(InternalReplicationIndex, EndReplicationFlags);
+	}
 }
 
 void UReplicationBridge::InternalAddSubObject(FNetRefHandle OwnerHandle, FNetRefHandle SubObjectHandle, FNetRefHandle InsertRelativeToSubObjectHandle, UE::Net::ESubObjectInsertionOrder InsertionOrder)
@@ -560,9 +563,10 @@ void UReplicationBridge::PreUpdateHandlesPendingEndReplication()
 
 	for (const FPendingEndReplicationInfo& Info : MakeArrayView(HandlesPendingEndReplication))
 	{
-		// Due to async stop replication not detaching instances in the call to StopReplication  we need to explicitly check for invalidated instances as objects might have been marked for garbage but GC has not been run.
-		if (FInternalNetRefIndex ObjectInternalIndex = NetRefHandleManager->GetInternalIndex(Info.Handle))
+		const FInternalNetRefIndex ObjectInternalIndex = Info.ObjectIndex;
+		if (NetRefHandleManager->GetAssignedInternalIndices().GetBit(ObjectInternalIndex))
 		{
+			// Due to async stop replication not explicitly detaching instances in the call to StopReplication we detach the instance protocol if the instance has been invalidated
 			if (!IsValid(NetRefHandleManager->GetReplicatedInstances()[ObjectInternalIndex]))
 			{
 				FNetRefHandleManager::FReplicatedObjectData& ObjectData = NetRefHandleManager->GetReplicatedObjectDataNoCheck(ObjectInternalIndex);
@@ -570,7 +574,8 @@ void UReplicationBridge::PreUpdateHandlesPendingEndReplication()
 				{
 					// Detach instance as we should not access the instance.
 					ObjectData.bPendingEndReplication = 1U;
-					InternalDetachInstanceFromNetRefHandle(ObjectInternalIndex);
+					const bool bUnbindInstanceProtocol = false;
+					InternalDetachInstanceFromNetRefHandle(ObjectInternalIndex, bUnbindInstanceProtocol);
 				}
 			}
 		}
diff --git a/Engine/Source/Runtime/Net/Iris/Public/Iris/ReplicationSystem/ReplicationBridge.h b/Engine/Source/Runtime/Net/Iris/Public/Iris/ReplicationSystem/ReplicationBridge.h
--- a/Engine/Source/Runtime/Net/Iris/Public/Iris/ReplicationSystem/ReplicationBridge.h
+++ b/Engine/Source/Runtime/Net/Iris/Public/Iris/ReplicationSystem/ReplicationBridge.h
@@ -219,7 +219,7 @@ protected:
 	IRISCORE_API void InternalAttachInstanceToNetRefHandle(FNetRefHandle RefHandle, bool bBindInstanceProtocol, UE::Net::FReplicationInstanceProtocol* InstanceProtocol, UObject* Instance, FNetHandle NetHandle);
 
 	/** Detach instance from NetRefHandle and destroy the instance protocol. */
-	IRISCORE_API void InternalDetachInstanceFromNetRefHandle(UE::Net::Private::FInternalNetRefIndex InternalObjectIndex);
+	IRISCORE_API void InternalDetachInstanceFromNetRefHandle(UE::Net::Private::FInternalNetRefIndex InternalObjectIndex, bool bUnBindInstanceProtocol = true);
 
 	/** Destroy the handle and all internal book keeping associated with it. */
 	IRISCORE_API void InternalDestroyNetObject(UE::Net::Private::FInternalNetRefIndex InternalObjectIndex);
