Iris - Fixed hysteresis bugs for late joining connections.

#Synced-CL 49914824

diff --git a/Engine/Plugins/Runtime/ReplicationSystemTestPlugin/Source/Private/Tests/ReplicationSystem/Filtering/TestFiltering.cpp b/Engine/Plugins/Runtime/ReplicationSystemTestPlugin/Source/Private/Tests/ReplicationSystem/Filtering/TestFiltering.cpp
--- a/Engine/Plugins/Runtime/ReplicationSystemTestPlugin/Source/Private/Tests/ReplicationSystem/Filtering/TestFiltering.cpp
+++ b/Engine/Plugins/Runtime/ReplicationSystemTestPlugin/Source/Private/Tests/ReplicationSystem/Filtering/TestFiltering.cpp
@@ -2725,6 +2725,59 @@ UE_NET_TEST_FIXTURE(FTestFilteringFixture, LotsOfObjectsCanBeFilteredOutViaHyste
 	}
 }
 
+// Test that a dynamically filtered out object doesn't enter hysteresis for late joining client 
+UE_NET_TEST_FIXTURE(FTestFilteringFixture, HysteresisDoesNotAffectDynamicallyFilteredOutObjectForLateJoiningClient)
+{
+	constexpr uint32 DefaultHysteresisFrameCount = 3;
+	FScopedDefaultHysteresisFrameCount ScopedDefaultHysteresisFrameCount(DefaultHysteresisFrameCount);
+
+	// Create object with filter which disallows replication.
+	UReplicatedTestObject* ServerObject = Server->CreateObject();
+	Server->ReplicationSystem->SetFilter(ServerObject->NetRefHandle, NotRoutedFilterHandle);
+
+	Server->UpdateAndSend({});
+
+	// Add late joining client
+	FReplicationSystemTestClient* Client = CreateClient();
+
+	// Perform send update for more than the hysteresis duration and make sure the object never gets created.
+	for (uint32 It = 0; It < DefaultHysteresisFrameCount + 1U; ++It)
+	{
+		Server->UpdateAndSend({ Client });
+		UE_NET_ASSERT_EQ(Client->GetReplicationBridge()->GetReplicatedObject(ServerObject->NetRefHandle), nullptr);
+	}
+}
+
+// Test that a dynamically filtered out dependent objects whose parent is not in scope doesn't enter hysteresis for late joining client 
+UE_NET_TEST_FIXTURE(FTestFilteringFixture, HysteresisDoesNotAffectDynamicallyFilteredOutDependentObjectForLateJoiningClient)
+{
+	constexpr uint32 DefaultHysteresisFrameCount = 3;
+	FScopedDefaultHysteresisFrameCount ScopedDefaultHysteresisFrameCount(DefaultHysteresisFrameCount);
+
+	// Create object with filter which disallows replication.
+	UReplicatedTestObject* ServerObject = Server->CreateObject();
+	Server->ReplicationSystem->SetFilter(ServerObject->NetRefHandle, NotRoutedFilterHandle);
+
+	// Create dependent object with filter which disallows replication.
+	UReplicatedTestObject* ServerDependentObject = Server->CreateObject();
+	Server->ReplicationSystem->SetFilter(ServerDependentObject->NetRefHandle, NotRoutedFilterHandle);
+
+	Server->ReplicationBridge->AddDependentObject(ServerObject->NetRefHandle, ServerDependentObject->NetRefHandle);
+
+	Server->UpdateAndSend({});
+
+	// Add late joining client
+	FReplicationSystemTestClient* Client = CreateClient();
+
+	// Perform send update for more than the hysteresis duration and make sure the objects never get created.
+	for (uint32 It = 0; It < DefaultHysteresisFrameCount + 1U; ++It)
+	{
+		Server->UpdateAndSend({ Client });
+		UE_NET_ASSERT_EQ(Client->GetReplicationBridge()->GetReplicatedObject(ServerObject->NetRefHandle), nullptr);
+		UE_NET_ASSERT_EQ(Client->GetReplicationBridge()->GetReplicatedObject(ServerDependentObject->NetRefHandle), nullptr);
+	}
+}
+
 UE_NET_TEST_FIXTURE(FTestFilteringFixture, TestDependentObjectIsFilteredOutTogetherWithExclusionFilteredOutParent)
 {
 	constexpr uint32 DefaultHysteresisFrameCount = 2;
diff --git a/Engine/Source/Runtime/Net/Iris/Private/Iris/ReplicationSystem/Filtering/ReplicationFiltering.cpp b/Engine/Source/Runtime/Net/Iris/Private/Iris/ReplicationSystem/Filtering/ReplicationFiltering.cpp
--- a/Engine/Source/Runtime/Net/Iris/Private/Iris/ReplicationSystem/Filtering/ReplicationFiltering.cpp
+++ b/Engine/Source/Runtime/Net/Iris/Private/Iris/ReplicationSystem/Filtering/ReplicationFiltering.cpp
@@ -355,6 +355,9 @@ void FReplicationFiltering::Filter()
 	}
 
 	FilterNonRelevantObjects();
+
+	// New connections are exempt from object scope hysteresis so we keep the info around until we're done with all filtering.
+	NewConnections.ClearAllBits();
 }
 
 void FReplicationFiltering::FilterNonRelevantObjects()
@@ -792,7 +795,6 @@ void FReplicationFiltering::InitNewConnections()
 	};
 
 	NewConnections.ForAllSetBits(InitNewConnection);
-	NewConnections.ClearAllBits();
 }
 
 void FReplicationFiltering::ResetRemovedConnections()
@@ -1508,21 +1510,11 @@ void FReplicationFiltering::UpdateDynamicFiltering()
 				}
 			}
 
-			// Add filtered out objects to hysteresis if eligible, meaning they support it and are still dynamically filtered. Dependent objects are not processed here.
-			for (const uint32 ObjectIndex : FilteredOutObjects)
+			if (const bool bConnIsExemptFromHysteresis = NewConnections.GetBit(ConnId))
 			{
-				const uint32 HysteresisFrameCount = ObjectScopeHysteresisFrameCounts[ObjectIndex];
-				const bool bAlreadyFilteredOut = DynamicFilteredOutObjectsHysteresisAdjusted.GetBit(ObjectIndex);
-				// If the object is already filtered out we cannot add to hysteresis. It may be that a subobject was added and the root object was marked for processing, in which case we should immediately filter out the subobject.
-				if (!bAlreadyFilteredOut && HysteresisFrameCount && DynamicFilterEnabledObjects.GetBit(ObjectIndex) && !HysteresisState.ObjectsExemptFromHysteresis.GetBit(ObjectIndex))
+				// If it's a new connection we should not apply hysteresis to any object- filter out objects not in scope immediately.
+				for (const uint32 ObjectIndex : FilteredOutObjects)
 				{
-					// We need to adjust the hysteresis frame count to account for when it will be updated. The -1 stems from the fact that updating won't happen until next frame at the earliest.
-					const uint16 TotalHysteresisFrameCount = static_cast<uint16>(HysteresisFrameCount - 1 + AdjustHysteresisForUpdateThrottling);
-					ConnectionInfo.HysteresisUpdater.SetHysteresisFrameCount(ObjectIndex, TotalHysteresisFrameCount);
-				}
-				else
-				{
-					// Filter out
 					DynamicFilteredOutObjectsHysteresisAdjusted.SetBit(ObjectIndex);
 					for (const FInternalNetRefIndex SubObjectIndex : NetRefHandleManager->GetSubObjects(ObjectIndex))
 					{
@@ -1530,6 +1522,31 @@ void FReplicationFiltering::UpdateDynamicFiltering()
 					}
 				}
 			}
+			else
+			{
+				// Add filtered out objects to hysteresis if eligible, meaning they support it and are still dynamically filtered. Dependent objects are not processed here.
+				for (const uint32 ObjectIndex : FilteredOutObjects)
+				{
+					const uint32 HysteresisFrameCount = ObjectScopeHysteresisFrameCounts[ObjectIndex];
+					const bool bAlreadyFilteredOut = DynamicFilteredOutObjectsHysteresisAdjusted.GetBit(ObjectIndex);
+					// If the object is already filtered out we cannot add to hysteresis. It may be that a subobject was added and the root object was marked for processing, in which case we should immediately filter out the subobject.
+					if (!bAlreadyFilteredOut && HysteresisFrameCount && DynamicFilterEnabledObjects.GetBit(ObjectIndex) && !HysteresisState.ObjectsExemptFromHysteresis.GetBit(ObjectIndex))
+					{
+						// We need to adjust the hysteresis frame count to account for when it will be updated. The -1 stems from the fact that updating won't happen until next frame at the earliest.
+						const uint16 TotalHysteresisFrameCount = static_cast<uint16>(HysteresisFrameCount - 1 + AdjustHysteresisForUpdateThrottling);
+						ConnectionInfo.HysteresisUpdater.SetHysteresisFrameCount(ObjectIndex, TotalHysteresisFrameCount);
+					}
+					else
+					{
+						// Filter out
+						DynamicFilteredOutObjectsHysteresisAdjusted.SetBit(ObjectIndex);
+						for (const FInternalNetRefIndex SubObjectIndex : NetRefHandleManager->GetSubObjects(ObjectIndex))
+						{
+							DynamicFilteredOutObjectsHysteresisAdjusted.SetBit(SubObjectIndex);
+						}
+					}
+				}
+			}
 
 			// Update hysteresis
 			if (ConnIdMod == HysteresisState.ConnectionStartId)
@@ -1619,6 +1636,8 @@ void FReplicationFiltering::UpdateDynamicFiltering()
 		{
 			IRIS_CSV_PROFILER_SCOPE(Iris, FReplicationFiltering_DependentObjects_Hysteresis);
 
+			const bool bConnIsExemptFromHysteresis = NewConnections.GetBit(ConnId);
+
 			// We didn't filter out any dependent objects yet so we enable/disable as needed.
 			// Different algorithm for figuring out update throttling adjustment since we already updated this frame.
 			const uint32 AdjustHysteresisForUpdateThrottling = (HysteresisState.ConnectionStartId + HysteresisState.ConnectionIdStride - ConnIdMod) % HysteresisState.ConnectionIdStride;
@@ -1643,7 +1662,7 @@ void FReplicationFiltering::UpdateDynamicFiltering()
 				else
 				{
 					// Start/continue hysteresis or time to prevent replication?
-					bool bIsFilteredOut = DynamicFilteredOutObjectsHysteresisAdjusted.GetBit(DependentObjectIndex) || !ConnectionInfo.ObjectsInScopeBeforeDynamicFiltering.GetBit(DependentObjectIndex);
+					bool bIsFilteredOut = bConnIsExemptFromHysteresis || DynamicFilteredOutObjectsHysteresisAdjusted.GetBit(DependentObjectIndex) || !ConnectionInfo.ObjectsInScopeBeforeDynamicFiltering.GetBit(DependentObjectIndex);
 					if (!bIsFilteredOut)
 					{
 						// If we're updated by hysteresis at this point we should not start again, but if other conditions allow it we should start!
