Fix Iris protocol name profiling.
Reverted IRIS_PROFILER_PROTOCOL_NAME macros to use the TEXT trace macros instead of the STR trace macros. The IRIS_PROFILER_PROTOCOL_NAME macros are used multiple times with different protocol names, so the caching done in the STR macros causes incorrect results in the trace.

#Synced-CL 50100874

#Synced-CL 50100874

diff --git a/Engine/Source/Runtime/Net/Iris/Public/Iris/Core/IrisProfiler.h b/Engine/Source/Runtime/Net/Iris/Public/Iris/Core/IrisProfiler.h
--- a/Engine/Source/Runtime/Net/Iris/Public/Iris/Core/IrisProfiler.h
+++ b/Engine/Source/Runtime/Net/Iris/Public/Iris/Core/IrisProfiler.h
@@ -61,8 +61,8 @@
 #endif
 
 #if UE_IRIS_PROFILER_ENABLE_PROTOCOL_NAMES
-#	define IRIS_PROFILER_PROTOCOL_NAME(X) IRIS_PROFILER_SCOPE_STR(X)
-#	define IRIS_PROFILER_PROTOCOL_NAME_CONDITIONAL(X, Cond) IRIS_PROFILER_SCOPE_STR_CONDITIONAL(X, Cond)
+#	define IRIS_PROFILER_PROTOCOL_NAME(X) IRIS_PROFILER_SCOPE_TEXT(X)
+#	define IRIS_PROFILER_PROTOCOL_NAME_CONDITIONAL(X, Cond) IRIS_PROFILER_SCOPE_TEXT_CONDITIONAL(X, Cond)
 #else
 #	define IRIS_PROFILER_PROTOCOL_NAME(X)
 #	define IRIS_PROFILER_PROTOCOL_NAME_CONDITIONAL(X, Cond)
