Iris
* Clarify document for NetObjectFactory::OnSubObjectCreatedFromReplication

#rnx

#Synced-CL 50152864

diff --git a/Engine/Source/Runtime/Net/Iris/Public/Iris/ReplicationSystem/NetObjectFactory.h b/Engine/Source/Runtime/Net/Iris/Public/Iris/ReplicationSystem/NetObjectFactory.h
--- a/Engine/Source/Runtime/Net/Iris/Public/Iris/ReplicationSystem/NetObjectFactory.h
+++ b/Engine/Source/Runtime/Net/Iris/Public/Iris/ReplicationSystem/NetObjectFactory.h
@@ -175,7 +175,8 @@ public:
 	virtual void DetachedFromReplication(const FDetachContext& Context, const TOptional<FSubObjectDetachContext>& SubObjectContext) PURE_VIRTUAL(UNetObjectFactory::DetachedFromReplication,);
 
 	/** 
-	 * Optional callback triggered when a root object managed by this factory gets assigned a dynamic subobject. 
+	 * Optional callback triggered when a root object managed by this factory gets assigned a subobject. 
+	 * The subobject factory must have set EReplicationBridgeCreateNetRefHandleResultFlags::ShouldCallSubObjectCreatedFromReplication in it's results. The default factories usually set the flag only for dynamic (not loaded) subobjects.
 	 * This callback is called on remotes only. 
 	 * At the time of the callback the RootObject will have the latest replicated properties set, but the subobject will only be default constructed and won't be assigned the replicated properties received alongside the creation request. 
 	 * @param RootObject The root object that owns the subobject
