SceneGraph
* Quick fix to force components to always be replicated first before any other subobjets
* Also ensure components are replicated in the order they are added in

Iris:
* Added insert option InsertAtStartOrAfterOther that allows you to replicate a subobject first or after a specific subobject

#rnx

#Synced-CL 50391842

#Synced-CL 50391842

#Synced-CL 50391842

diff --git a/Engine/Source/Runtime/Net/Iris/Private/Iris/ReplicationSystem/NetRefHandleManager.cpp b/Engine/Source/Runtime/Net/Iris/Private/Iris/ReplicationSystem/NetRefHandleManager.cpp
--- a/Engine/Source/Runtime/Net/Iris/Private/Iris/ReplicationSystem/NetRefHandleManager.cpp
+++ b/Engine/Source/Runtime/Net/Iris/Private/Iris/ReplicationSystem/NetRefHandleManager.cpp
@@ -769,7 +769,7 @@ bool FNetRefHandleManager::AddSubObject(FNetRefHandle RootObjectHandle, FNetRefH
 
 	UE_NET_TRACE_SUBOBJECT(RootObjectHandle, SubObjectHandle);
 
-	const FInternalNetRefIndex RelativeOtherSubObjectInternalIndex = EnumHasAnyFlags(Flags, EAddSubObjectFlags::ReplicateWithSubObject) ? GetInternalIndex(RelativeOtherSubObjectHandle) : InvalidInternalIndex;
+	const FInternalNetRefIndex RelativeOtherSubObjectInternalIndex = EnumHasAnyFlags(Flags, EAddSubObjectFlags::ReplicateWithSubObject | EAddSubObjectFlags::InsertAtStartOrAfterOther) ? GetInternalIndex(RelativeOtherSubObjectHandle) : InvalidInternalIndex;
 	return InternalAddSubObject(RootObjectInternalIndex, SubObjectInternalIndex, RelativeOtherSubObjectInternalIndex, Flags);
 }
 
@@ -791,6 +791,14 @@ bool FNetRefHandleManager::InternalAddSubObject(FInternalNetRefIndex RootObjectI
 		// Add at the start
 		SubObjectArray.Insert(SubObjectInternalIndex, 0);
 	}
+	else if (EnumHasAnyFlags(Flags, EAddSubObjectFlags::InsertAtStartOrAfterOther))
+	{
+		// Insert the object right after the other one, or at the start if it wasn't in the list
+		const int32 RelativeIndex = RelativeOtherSubObjectInternalIndex != FNetRefHandleManager::InvalidInternalIndex ? SubObjectArray.Find(RelativeOtherSubObjectInternalIndex) : INDEX_NONE;
+		const int32 InsertIndex = (RelativeIndex == INDEX_NONE) ? 0 : (RelativeIndex + 1);
+
+		SubObjectArray.Insert(SubObjectInternalIndex, InsertIndex);
+	}
 	else
 	{
 		// At at the end
@@ -805,7 +813,7 @@ bool FNetRefHandleManager::InternalAddSubObject(FInternalNetRefIndex RootObjectI
 	FInternalNetRefIndex ParentOfSubObjectIndex = RootObjectInternalIndex;
 	
 	// Make sure the non-rootobject Parent is a subobject of the same RootObject
-	if (RelativeOtherSubObjectInternalIndex != InvalidInternalIndex)
+	if (EnumHasAnyFlags(Flags, EAddSubObjectFlags::ReplicateWithSubObject) && RelativeOtherSubObjectInternalIndex != InvalidInternalIndex)
 	{
 		const bool bIsValidOuter = SubObjectArray.Contains(RelativeOtherSubObjectInternalIndex);
 		if (ensureMsgf(bIsValidOuter, TEXT("RelativeOtherSubObjectHandle %s must be a Subobject of %s"), *PrintObjectFromIndex(RelativeOtherSubObjectInternalIndex), *PrintObjectFromIndex(RootObjectInternalIndex)))
@@ -823,6 +831,14 @@ bool FNetRefHandleManager::InternalAddSubObject(FInternalNetRefIndex RootObjectI
 		// Add at the start
 		ChildSubObjectArray.Insert(SubObjectInternalIndex, 0);
 	}
+	else if (EnumHasAnyFlags(Flags, EAddSubObjectFlags::InsertAtStartOrAfterOther))
+	{
+		// Insert the object right after the other one, or at the start if it wasn't in the list
+		const int32 RelativeIndex = RelativeOtherSubObjectInternalIndex != FNetRefHandleManager::InvalidInternalIndex ? ChildSubObjectArray.Find(RelativeOtherSubObjectInternalIndex) : INDEX_NONE;
+		const int32 InsertIndex = (RelativeIndex == INDEX_NONE) ? 0 : (RelativeIndex + 1);
+
+		ChildSubObjectArray.Insert(SubObjectInternalIndex, InsertIndex);
+	}
 	else
 	{
 		// At at the end
diff --git a/Engine/Source/Runtime/Net/Iris/Private/Iris/ReplicationSystem/NetRefHandleManager.h b/Engine/Source/Runtime/Net/Iris/Private/Iris/ReplicationSystem/NetRefHandleManager.h
--- a/Engine/Source/Runtime/Net/Iris/Private/Iris/ReplicationSystem/NetRefHandleManager.h
+++ b/Engine/Source/Runtime/Net/Iris/Private/Iris/ReplicationSystem/NetRefHandleManager.h
@@ -51,8 +51,10 @@ enum class EAddSubObjectFlags : uint32
 {
 	None = 0U,
 	ReplicateWithSubObject = 1U,
-	// Insert the subject at the start of the list
+	/** Insert the subject at the start of the list */
 	InsertAtStart = ReplicateWithSubObject << 1U, 
+	/** Insert at the start of the list or right after the relative handle if one is set */
+	InsertAtStartOrAfterOther = InsertAtStart << 1U,
 };
 ENUM_CLASS_FLAGS(EAddSubObjectFlags);
 
diff --git a/Engine/Source/Runtime/Net/Iris/Private/Iris/ReplicationSystem/ReplicationBridge.cpp b/Engine/Source/Runtime/Net/Iris/Private/Iris/ReplicationSystem/ReplicationBridge.cpp
--- a/Engine/Source/Runtime/Net/Iris/Private/Iris/ReplicationSystem/ReplicationBridge.cpp
+++ b/Engine/Source/Runtime/Net/Iris/Private/Iris/ReplicationSystem/ReplicationBridge.cpp
@@ -349,6 +349,9 @@ void UReplicationBridge::InternalAddSubObject(FNetRefHandle OwnerHandle, FNetRef
 		case ESubObjectInsertionOrder::InsertAtStart:
 			AddSubObjectFlags |= EAddSubObjectFlags::InsertAtStart;
 			break;
+		case ESubObjectInsertionOrder::InsertAtStartOrAfterOther:
+			AddSubObjectFlags |= EAddSubObjectFlags::InsertAtStartOrAfterOther;
+			break;
 		default:
 			checkf(false, TEXT("Missing implementation of ESubObjectInsertionOrder enum"));
 			break;
diff --git a/Engine/Source/Runtime/Net/Iris/Public/Iris/ReplicationSystem/ReplicationBridge.h b/Engine/Source/Runtime/Net/Iris/Public/Iris/ReplicationSystem/ReplicationBridge.h
--- a/Engine/Source/Runtime/Net/Iris/Public/Iris/ReplicationSystem/ReplicationBridge.h
+++ b/Engine/Source/Runtime/Net/Iris/Public/Iris/ReplicationSystem/ReplicationBridge.h
@@ -102,6 +102,8 @@ namespace UE::Net
 		ReplicateWith,
 		/** Insert the subobject at the start of the list so it can be created and replicated first */
 		InsertAtStart,
+		/** Insert at the start of the list or right after the relative handle if one is set */
+		InsertAtStartOrAfterOther,
 	};
 
 } // end namespace UE::Net
