Iris - Allow options to be passed to UE_NET_IMPLEMENT_NAMED_STRUCT_LASTRESORT_NETSERIALIZER_INFO and UE_NET_IMPLEMENT_NAMED_STRUCT_LASTRESORT_NETSERIALIZER_AND_REGISTRY_DELEGATES. Current options are to set the max quantized size, exclude from default state hash and custom fragment registration.

#Synced-CL 50441695

diff --git a/Engine/Source/Runtime/Net/Iris/Public/Iris/ReplicationState/PropertyNetSerializerInfoRegistry.h b/Engine/Source/Runtime/Net/Iris/Public/Iris/ReplicationState/PropertyNetSerializerInfoRegistry.h
--- a/Engine/Source/Runtime/Net/Iris/Public/Iris/ReplicationState/PropertyNetSerializerInfoRegistry.h
+++ b/Engine/Source/Runtime/Net/Iris/Public/Iris/ReplicationState/PropertyNetSerializerInfoRegistry.h
@@ -3,6 +3,7 @@
 #pragma once
 
 #include "Containers/Array.h"
+#include "HAL/PreprocessorHelpers.h"
 #include "Iris/Serialization/NetSerializer.h"
 #include "Iris/Serialization/NetSerializerDelegates.h"
 #include "Templates/Casts.h"
@@ -165,7 +166,34 @@ struct FLastResortPropertyNetSerializerInfo : public FPropertyNetSerializerInfo
 	IRISCORE_API virtual bool CanUseDefaultConfig(const FProperty* Property) const override final;
 	IRISCORE_API virtual FNetSerializerConfig* BuildNetSerializerConfig(void* NetSerializerConfigBuffer, const FProperty* Property) const override final;
 
+	virtual CreateAndRegisterReplicationFragmentFunc GetCreateAndRegisterReplicationFragmentFunction() const override final
+	{
+		return CreateAndRegisterReplicationFragmentFunction;
+	}
+
+	// Options to modify the behavior for a particular named struct
+
+	// Pass UE_NET_LASTRESORT_NETSERIALIZER_OPTION_MAXSIZE(DesiredMaxSize) to UE_NET_IMPLEMENT_NAMED_STRUCT_LASTRESORT_NETSERIALIZER_INFO macro
+	void SetMaxQuantizedSizeBits(uint32 InMaxQuantizedSizeBits)
+	{
+		MaxQuantizedSizeBits = InMaxQuantizedSizeBits;
+	}
+
+	// Pass UE_NET_LASTRESORT_NETSERIALIZER_OPTION_EXCLUDEFROMDEFAULTSTATEHASH(bExclude) to UE_NET_IMPLEMENT_NAMED_STRUCT_LASTRESORT_NETSERIALIZER_INFO macro
+	void SetExcludeFromDefaultStateHash(bool bExclude)
+	{
+		bExcludeFromDefaultStateHash = bExclude;
+	}
+
+	// Pass UE_NET_LASTRESORT_NETSERIALIZER_OPTION_CREATEANDREGISTERREPLICATIOFRAGMENTFUNCTION(CreateAndRegisterReplicationFragmentFunction) to UE_NET_IMPLEMENT_NAMED_STRUCT_LASTRESORT_NETSERIALIZER_INFO macro
+	void SetCreateAndRegisterReplicationFragmentFunction(CreateAndRegisterReplicationFragmentFunc InCreateAndRegisterReplicationFragmentFunction)
+	{
+		CreateAndRegisterReplicationFragmentFunction = InCreateAndRegisterReplicationFragmentFunction;
+	}
+
 protected:
+	CreateAndRegisterReplicationFragmentFunc CreateAndRegisterReplicationFragmentFunction = nullptr;
+
 	// Maximum size of quantized state in bits
 	uint32 MaxQuantizedSizeBits;
 
@@ -232,13 +260,6 @@ const UE::Net::FPropertyNetSerializerInfo& GetPropertyNetSerializerInfo_##Name()
 	return StaticInstance; \
 };
 
-// Implement FPropertyNetSerializerInfo for cases where LastResortPropertyNetSerializer is needed for struct with custom serialization.
-#define UE_NET_IMPLEMENT_NAMED_STRUCT_LASTRESORT_NETSERIALIZER_INFO(StructName) \
-const UE::Net::FPropertyNetSerializerInfo& GetPropertyNetSerializerInfo_##StructName() { static UE::Net::FNamedStructLastResortPropertyNetSerializerInfo StaticInstance(StructName); return StaticInstance; };
-
-// Implement FPropertyNetSerializerInfo for cases where LastResortPropertyNetSerializer is needed for struct with custom serialization, with quantized state size override.
-#define UE_NET_IMPLEMENT_NAMED_STRUCT_LASTRESORT_NETSERIALIZER_INFO_WITH_SIZE_OVERRIDE(StructName, MaxBits) \
-const UE::Net::FPropertyNetSerializerInfo& GetPropertyNetSerializerInfo_##StructName() { static UE::Net::FNamedStructLastResortPropertyNetSerializerInfo StaticInstance(StructName, MaxBits); return StaticInstance; };
 
 // Register 
 #define UE_NET_REGISTER_NETSERIALIZER_INFO(Name) \
@@ -268,11 +289,45 @@ static F##Name##NetSerializerRegistryDelegates Name##NetSerializerRegistryDelega
 	UE_NET_IMPLEMENT_NETSERIALIZER_REGISTRY_DELEGATES(Name)
 
 
-// Utility that can be used to forward serialization of a Struct to a last resort net serializer
-#define UE_NET_IMPLEMENT_NAMED_STRUCT_LASTRESORT_NETSERIALIZER_AND_REGISTRY_DELEGATES(Name) \
+// Helper macros that allows us to pass a variable number of arguments, up to 8, to modify LastResortNetSerializer behavior for a particular named struct
+#define UE_NET_PRIVATE_FOR_EACH_0(...) 
+#define UE_NET_PRIVATE_FOR_EACH_1(ApplyMacro, InstanceName, OptionMacro)      ApplyMacro(InstanceName, OptionMacro)
+#define UE_NET_PRIVATE_FOR_EACH_2(ApplyMacro, InstanceName, OptionMacro, ...) ApplyMacro(InstanceName, OptionMacro) UE_NET_PRIVATE_FOR_EACH_1(ApplyMacro, InstanceName, __VA_ARGS__)
+#define UE_NET_PRIVATE_FOR_EACH_3(ApplyMacro, InstanceName, OptionMacro, ...) ApplyMacro(InstanceName, OptionMacro) UE_NET_PRIVATE_FOR_EACH_2(ApplyMacro, InstanceName, __VA_ARGS__)
+#define UE_NET_PRIVATE_FOR_EACH_4(ApplyMacro, InstanceName, OptionMacro, ...) ApplyMacro(InstanceName, OptionMacro) UE_NET_PRIVATE_FOR_EACH_3(ApplyMacro, InstanceName, __VA_ARGS__)
+#define UE_NET_PRIVATE_FOR_EACH_5(ApplyMacro, InstanceName, OptionMacro, ...) ApplyMacro(InstanceName, OptionMacro) UE_NET_PRIVATE_FOR_EACH_4(ApplyMacro, InstanceName, __VA_ARGS__)
+#define UE_NET_PRIVATE_FOR_EACH_6(ApplyMacro, InstanceName, OptionMacro, ...) ApplyMacro(InstanceName, OptionMacro) UE_NET_PRIVATE_FOR_EACH_5(ApplyMacro, InstanceName, __VA_ARGS__)
+#define UE_NET_PRIVATE_FOR_EACH_7(ApplyMacro, InstanceName, OptionMacro, ...) ApplyMacro(InstanceName, OptionMacro) UE_NET_PRIVATE_FOR_EACH_6(ApplyMacro, InstanceName, __VA_ARGS__)
+#define UE_NET_PRIVATE_FOR_EACH_8(ApplyMacro, InstanceName, OptionMacro, ...) ApplyMacro(InstanceName, OptionMacro) UE_NET_PRIVATE_FOR_EACH_7(ApplyMacro, InstanceName, __VA_ARGS__)
+#define UE_NET_PRIVATE_FOR_EACH(ApplyMacro, InstanceName, ...) UE_APPEND_VA_ARG_COUNT(UE_NET_PRIVATE_FOR_EACH_ __VA_OPT__(,) __VA_ARGS__)(ApplyMacro, InstanceName, __VA_ARGS__)
+#define UE_NET_PRIVATE_APPLY_OPTION(InstanceName, OptionMacro) InstanceName.OptionMacro;
+
+// LastResortPropertyNetSerializer options. May only be passed as arguments to UE_NET_IMPLEMENT_NAMED_STRUCT_LASTRESORT_NETSERIALIZER_INFO
+#define UE_NET_LASTRESORT_NETSERIALIZER_OPTION_MAXSIZE(MaxSize) SetMaxQuantizedSizeBits(MaxSize)
+#define UE_NET_LASTRESORT_NETSERIALIZER_OPTION_EXCLUDEFROMDEFAULTSTATEHASH(bExclude) SetExcludeFromDefaultStateHash(bExclude)
+#define UE_NET_LASTRESORT_NETSERIALIZER_OPTION_CREATEANDREGISTERREPLICATIOFRAGMENTFUNCTION(CreateAndRegisterReplicationFragmentFunction) SetCreateAndRegisterReplicationFragmentFunction(CreateAndRegisterReplicationFragmentFunction)
+
+// Implement FPropertyNetSerializerInfo for cases where LastResortPropertyNetSerializer is needed for struct with custom serialization. You need to pass an FName as StructName. Optionally pass applicable UE_NET_LASTRESORT_NETSERIALIZER_OPTION_ options.
+#define UE_NET_IMPLEMENT_NAMED_STRUCT_LASTRESORT_NETSERIALIZER_INFO(StructName, ... /* options */) \
+const UE::Net::FPropertyNetSerializerInfo& GetPropertyNetSerializerInfo_##StructName() \
+{ \
+	static UE::Net::FNamedStructLastResortPropertyNetSerializerInfo StaticInstance(StructName); \
+	UE_NET_PRIVATE_FOR_EACH(UE_NET_PRIVATE_APPLY_OPTION, StaticInstance, __VA_ARGS__) \
+	return StaticInstance; \
+}
+
+// Implement FPropertyNetSerializerInfo for cases where LastResortPropertyNetSerializer is needed for struct with custom serialization, with quantized state size override.
+#define UE_NET_IMPLEMENT_NAMED_STRUCT_LASTRESORT_NETSERIALIZER_INFO_WITH_SIZE_OVERRIDE(StructName, MaxBits) \
+	UE_NET_IMPLEMENT_NAMED_STRUCT_LASTRESORT_NETSERIALIZER_INFO(StructName, UE_NET_LASTRESORT_NETSERIALIZER_OPTION_MAXSIZE(MaxBits))
+
+/**
+ * Utility that can be used to forward serialization of a Struct to a last resort net serializer. Optionally pass applicable UE_NET_LASTRESORT_NETSERIALIZER_OPTION_ options.
+ * Example: UE_NET_IMPLEMENT_NAMED_STRUCT_LASTRESORT_NETSERIALIZER_AND_REGISTRY_DELEGATES(MyStruct, UE_NET_LASTRESORT_NETSERIALIZER_OPTION_MAXSIZE(32000), UE_NET_LASTRESORT_NETSERIALIZER_OPTION_EXCLUDEFROMDEFAULTSTATEHASH(true))
+ */
+#define UE_NET_IMPLEMENT_NAMED_STRUCT_LASTRESORT_NETSERIALIZER_AND_REGISTRY_DELEGATES(Name, ... /* options */) \
 	UE_NET_IS_DECLARED_TYPE(Name); \
 	static const FName PropertyNetSerializerRegistry_NAME_##Name( UE_STRINGIZE(Name) ); \
-	UE_NET_IMPLEMENT_NAMED_STRUCT_LASTRESORT_NETSERIALIZER_INFO(PropertyNetSerializerRegistry_NAME_##Name); \
+	UE_NET_IMPLEMENT_NAMED_STRUCT_LASTRESORT_NETSERIALIZER_INFO(PropertyNetSerializerRegistry_NAME_##Name __VA_OPT__(,) __VA_ARGS__) \
 	UE_NET_IMPLEMENT_NETSERIALIZER_REGISTRY_DELEGATES(Name)
 
 // Utility that can be used to forward serialization of a Struct to a last resort net serializer with a maximum quantized size
