Iris - Reworked how to pass options to UE_NET_IMPLEMENT_NAMED_STRUCT_LASTRESORT_NETSERIALIZER_INFO and UE_NET_IMPLEMENT_NAMED_STRUCT_LASTRESORT_NETSERIALIZER_AND_REGISTRY_DELEGATES to be more terse and user friendly. One can now pass an optional FLastResortPropertyNetSerializerUserOptions which allows overriding the max quantized size, custom fragment registration and whether to exclude the state from the replication state default state hash.

#Synced-CL 50441695

diff --git a/Engine/Source/Runtime/Net/Iris/Private/Iris/ReplicationState/PropertyNetSerializerInfoRegistry.cpp b/Engine/Source/Runtime/Net/Iris/Private/Iris/ReplicationState/PropertyNetSerializerInfoRegistry.cpp
--- a/Engine/Source/Runtime/Net/Iris/Private/Iris/ReplicationState/PropertyNetSerializerInfoRegistry.cpp
+++ b/Engine/Source/Runtime/Net/Iris/Private/Iris/ReplicationState/PropertyNetSerializerInfoRegistry.cpp
@@ -85,15 +85,11 @@ struct FDefaultArrayPropertyNetSerializerInfo : public FPropertyNetSerializerInf
 UE_NET_IMPLEMENT_NETSERIALIZER_INFO(FDefaultArrayPropertyNetSerializerInfo);
 
 /** Last resort for properties that have no specialized serializer. An example is TextProperty. */
-FLastResortPropertyNetSerializerInfo::FLastResortPropertyNetSerializerInfo()
-	: FPropertyNetSerializerInfo()
-	, MaxQuantizedSizeBits(FLastResortPropertyNetSerializerConfig::DefaultMaxQuantizedSizeBits)
-{
-}
-
-FLastResortPropertyNetSerializerInfo::FLastResortPropertyNetSerializerInfo(const uint32 InMaxQuantizedSizeBits)
-	: FPropertyNetSerializerInfo()
-	, MaxQuantizedSizeBits(InMaxQuantizedSizeBits)
+FLastResortPropertyNetSerializerInfo::FLastResortPropertyNetSerializerInfo(const FLastResortPropertyNetSerializerUserOptions& UserOptions)
+: FPropertyNetSerializerInfo()
+, CreateAndRegisterReplicationFragmentFunction(UserOptions.CreateAndRegisterReplicationFragmentFunction)
+, MaxQuantizedSizeBits(UserOptions.MaxQuantizedSizeBits)
+, bExcludeFromDefaultStateHash(UserOptions.bExcludeFromDefaultStateHash)
 {
 }
 
@@ -122,7 +118,7 @@ FNetSerializerConfig* FLastResortPropertyNetSerializerInfo::BuildNetSerializerCo
 	FLastResortPropertyNetSerializerConfig* Config = new (NetSerializerConfigBuffer) FLastResortPropertyNetSerializerConfig();
 	InitLastResortPropertyNetSerializerConfigFromProperty(*Config, Property);
 	Config->bExcludeFromDefaultStateHash = bExcludeFromDefaultStateHash;
-	Config->MaxQuantizedSizeBits = MaxQuantizedSizeBits;
+	Config->MaxQuantizedSizeBits = MaxQuantizedSizeBits > 0U ? MaxQuantizedSizeBits : FLastResortPropertyNetSerializerConfig::DefaultMaxQuantizedSizeBits;
 
 	return Config;
 }
diff --git a/Engine/Source/Runtime/Net/Iris/Public/Iris/ReplicationState/PropertyNetSerializerInfoRegistry.h b/Engine/Source/Runtime/Net/Iris/Public/Iris/ReplicationState/PropertyNetSerializerInfoRegistry.h
--- a/Engine/Source/Runtime/Net/Iris/Public/Iris/ReplicationState/PropertyNetSerializerInfoRegistry.h
+++ b/Engine/Source/Runtime/Net/Iris/Public/Iris/ReplicationState/PropertyNetSerializerInfoRegistry.h
@@ -70,7 +70,7 @@ protected:
 };
 
 /**
- *  This is a simple static registry used by the ReplicationStateDescriptorBuiider when building descriptors using properties
+ *  This is a simple static registry used by the ReplicationStateDescriptorBuilder when building descriptors using properties
  */
 class FPropertyNetSerializerInfoRegistry
 {
@@ -155,10 +155,21 @@ private:
 	CreateAndRegisterReplicationFragmentFunc CreateAndRegisterReplicationFragmentFunction = nullptr;
 };
 
+/** These are options for the LastResortPropertyNetSerializer that are user controllable by passing an instance to a registration macro, for example UE_NET_IMPLEMENT_NAMED_STRUCT_LASTRESORT_NETSERIALIZER_AND_REGISTRY_DELEGATES(MyStruct, {.MaxQuantizedSizeBits = 32767}) */
+struct FLastResortPropertyNetSerializerUserOptions
+{
+	/** If the struct needs a custom fragment due to for example not being able to dequantize to a temporary state. */
+	CreateAndRegisterReplicationFragmentFunc CreateAndRegisterReplicationFragmentFunction = nullptr;
+	/** How large the quantized state may be in bits. Only needed to be set if the requirements are very high. Leave as 0 to use whatever the default is (65535).  */
+	uint32 MaxQuantizedSizeBits = 0;
+	/** Whether to exclude the serialization from being part of a replication state's default state hash. This should be done only if the property may differ on the sending and receiving side. */
+	bool bExcludeFromDefaultStateHash = false;
+};
+
 struct FLastResortPropertyNetSerializerInfo : public FPropertyNetSerializerInfo
 {
-	IRISCORE_API FLastResortPropertyNetSerializerInfo();
-	IRISCORE_API explicit FLastResortPropertyNetSerializerInfo(const uint32 InMaxQuantizedSizeBits);
+	FLastResortPropertyNetSerializerInfo() = default;
+	IRISCORE_API explicit FLastResortPropertyNetSerializerInfo(const FLastResortPropertyNetSerializerUserOptions& UserOptions);
 
 	IRISCORE_API virtual const FFieldClass* GetPropertyTypeClass() const override;
 	IRISCORE_API virtual bool IsSupported(const FProperty* Property) const override;
@@ -171,31 +182,11 @@ struct FLastResortPropertyNetSerializerInfo : public FPropertyNetSerializerInfo
 		return CreateAndRegisterReplicationFragmentFunction;
 	}
 
-	// Options to modify the behavior for a particular named struct
-
-	// Pass UE_NET_LASTRESORT_NETSERIALIZER_OPTION_MAXSIZE(DesiredMaxSize) to UE_NET_IMPLEMENT_NAMED_STRUCT_LASTRESORT_NETSERIALIZER_INFO macro
-	void SetMaxQuantizedSizeBits(uint32 InMaxQuantizedSizeBits)
-	{
-		MaxQuantizedSizeBits = InMaxQuantizedSizeBits;
-	}
-
-	// Pass UE_NET_LASTRESORT_NETSERIALIZER_OPTION_EXCLUDEFROMDEFAULTSTATEHASH(bExclude) to UE_NET_IMPLEMENT_NAMED_STRUCT_LASTRESORT_NETSERIALIZER_INFO macro
-	void SetExcludeFromDefaultStateHash(bool bExclude)
-	{
-		bExcludeFromDefaultStateHash = bExclude;
-	}
-
-	// Pass UE_NET_LASTRESORT_NETSERIALIZER_OPTION_CREATEANDREGISTERREPLICATIOFRAGMENTFUNCTION(CreateAndRegisterReplicationFragmentFunction) to UE_NET_IMPLEMENT_NAMED_STRUCT_LASTRESORT_NETSERIALIZER_INFO macro
-	void SetCreateAndRegisterReplicationFragmentFunction(CreateAndRegisterReplicationFragmentFunc InCreateAndRegisterReplicationFragmentFunction)
-	{
-		CreateAndRegisterReplicationFragmentFunction = InCreateAndRegisterReplicationFragmentFunction;
-	}
-
 protected:
 	CreateAndRegisterReplicationFragmentFunc CreateAndRegisterReplicationFragmentFunction = nullptr;
 
-	// Maximum size of quantized state in bits
-	uint32 MaxQuantizedSizeBits;
+	// Maximum size of quantized state in bits. 0 means use the default value (65535)
+	uint32 MaxQuantizedSizeBits = 0;
 
 	// If set this class should not be included in the default state hash.
 	bool bExcludeFromDefaultStateHash = false;
@@ -203,14 +194,14 @@ protected:
 
 struct FNamedStructLastResortPropertyNetSerializerInfo final : public FLastResortPropertyNetSerializerInfo
 {
-	FNamedStructLastResortPropertyNetSerializerInfo(const FName InPropertyFName)
+	explicit FNamedStructLastResortPropertyNetSerializerInfo(const FName InPropertyFName)
 	: FLastResortPropertyNetSerializerInfo()
 	{
 		StructName = InPropertyFName;
 	}
 
-	FNamedStructLastResortPropertyNetSerializerInfo(const FName InPropertyFName, const uint32 InMaxBits)
-	: FLastResortPropertyNetSerializerInfo(InMaxBits)
+	FNamedStructLastResortPropertyNetSerializerInfo(const FName InPropertyFName, const FLastResortPropertyNetSerializerUserOptions& UserOptions)
+	: FLastResortPropertyNetSerializerInfo(UserOptions)
 	{
 		StructName = InPropertyFName;
 	}
@@ -288,43 +279,23 @@ static F##Name##NetSerializerRegistryDelegates Name##NetSerializerRegistryDelega
 	UE_NET_IMPLEMENT_NAMED_STRUCT_NETSERIALIZER_INFO(PropertyNetSerializerRegistry_NAME_##Name, SerializerName); \
 	UE_NET_IMPLEMENT_NETSERIALIZER_REGISTRY_DELEGATES(Name)
 
-
-// Helper macros that allows us to pass a variable number of arguments, up to 8, to modify LastResortNetSerializer behavior for a particular named struct
-#define UE_NET_PRIVATE_FOR_EACH_0(...) 
-#define UE_NET_PRIVATE_FOR_EACH_1(ApplyMacro, InstanceName, OptionMacro)      ApplyMacro(InstanceName, OptionMacro)
-#define UE_NET_PRIVATE_FOR_EACH_2(ApplyMacro, InstanceName, OptionMacro, ...) ApplyMacro(InstanceName, OptionMacro) UE_NET_PRIVATE_FOR_EACH_1(ApplyMacro, InstanceName, __VA_ARGS__)
-#define UE_NET_PRIVATE_FOR_EACH_3(ApplyMacro, InstanceName, OptionMacro, ...) ApplyMacro(InstanceName, OptionMacro) UE_NET_PRIVATE_FOR_EACH_2(ApplyMacro, InstanceName, __VA_ARGS__)
-#define UE_NET_PRIVATE_FOR_EACH_4(ApplyMacro, InstanceName, OptionMacro, ...) ApplyMacro(InstanceName, OptionMacro) UE_NET_PRIVATE_FOR_EACH_3(ApplyMacro, InstanceName, __VA_ARGS__)
-#define UE_NET_PRIVATE_FOR_EACH_5(ApplyMacro, InstanceName, OptionMacro, ...) ApplyMacro(InstanceName, OptionMacro) UE_NET_PRIVATE_FOR_EACH_4(ApplyMacro, InstanceName, __VA_ARGS__)
-#define UE_NET_PRIVATE_FOR_EACH_6(ApplyMacro, InstanceName, OptionMacro, ...) ApplyMacro(InstanceName, OptionMacro) UE_NET_PRIVATE_FOR_EACH_5(ApplyMacro, InstanceName, __VA_ARGS__)
-#define UE_NET_PRIVATE_FOR_EACH_7(ApplyMacro, InstanceName, OptionMacro, ...) ApplyMacro(InstanceName, OptionMacro) UE_NET_PRIVATE_FOR_EACH_6(ApplyMacro, InstanceName, __VA_ARGS__)
-#define UE_NET_PRIVATE_FOR_EACH_8(ApplyMacro, InstanceName, OptionMacro, ...) ApplyMacro(InstanceName, OptionMacro) UE_NET_PRIVATE_FOR_EACH_7(ApplyMacro, InstanceName, __VA_ARGS__)
-#define UE_NET_PRIVATE_FOR_EACH(ApplyMacro, InstanceName, ...) UE_APPEND_VA_ARG_COUNT(UE_NET_PRIVATE_FOR_EACH_ __VA_OPT__(,) __VA_ARGS__)(ApplyMacro, InstanceName, __VA_ARGS__)
-#define UE_NET_PRIVATE_APPLY_OPTION(InstanceName, OptionMacro) InstanceName.OptionMacro;
-
-// LastResortPropertyNetSerializer options. May only be passed as arguments to UE_NET_IMPLEMENT_NAMED_STRUCT_LASTRESORT_NETSERIALIZER_INFO
-#define UE_NET_LASTRESORT_NETSERIALIZER_OPTION_MAXSIZE(MaxSize) SetMaxQuantizedSizeBits(MaxSize)
-#define UE_NET_LASTRESORT_NETSERIALIZER_OPTION_EXCLUDEFROMDEFAULTSTATEHASH(bExclude) SetExcludeFromDefaultStateHash(bExclude)
-#define UE_NET_LASTRESORT_NETSERIALIZER_OPTION_CREATEANDREGISTERREPLICATIOFRAGMENTFUNCTION(CreateAndRegisterReplicationFragmentFunction) SetCreateAndRegisterReplicationFragmentFunction(CreateAndRegisterReplicationFragmentFunction)
-
-// Implement FPropertyNetSerializerInfo for cases where LastResortPropertyNetSerializer is needed for struct with custom serialization. You need to pass an FName as StructName. Optionally pass applicable UE_NET_LASTRESORT_NETSERIALIZER_OPTION_ options.
-#define UE_NET_IMPLEMENT_NAMED_STRUCT_LASTRESORT_NETSERIALIZER_INFO(StructName, ... /* options */) \
+// Implement FPropertyNetSerializerInfo for cases where LastResortPropertyNetSerializer is needed for struct with custom serialization. You need to pass an FName as StructName. Optionally pass FLastResortPropertyNetSerializerUserOptions.
+#define UE_NET_IMPLEMENT_NAMED_STRUCT_LASTRESORT_NETSERIALIZER_INFO(StructName, ... /* FLastResortPropertyNetSerializerUserOptions instance */) \
 const UE::Net::FPropertyNetSerializerInfo& GetPropertyNetSerializerInfo_##StructName() \
 { \
-	static UE::Net::FNamedStructLastResortPropertyNetSerializerInfo StaticInstance(StructName); \
-	UE_NET_PRIVATE_FOR_EACH(UE_NET_PRIVATE_APPLY_OPTION, StaticInstance, __VA_ARGS__) \
+	static UE::Net::FNamedStructLastResortPropertyNetSerializerInfo StaticInstance(StructName __VA_OPT__(,) __VA_ARGS__); \
 	return StaticInstance; \
 }
 
 // Implement FPropertyNetSerializerInfo for cases where LastResortPropertyNetSerializer is needed for struct with custom serialization, with quantized state size override.
 #define UE_NET_IMPLEMENT_NAMED_STRUCT_LASTRESORT_NETSERIALIZER_INFO_WITH_SIZE_OVERRIDE(StructName, MaxBits) \
-	UE_NET_IMPLEMENT_NAMED_STRUCT_LASTRESORT_NETSERIALIZER_INFO(StructName, UE_NET_LASTRESORT_NETSERIALIZER_OPTION_MAXSIZE(MaxBits))
+	UE_NET_IMPLEMENT_NAMED_STRUCT_LASTRESORT_NETSERIALIZER_INFO(StructName, {.MaxQuantizedSizeBits = MaxBits})
 
 /**
- * Utility that can be used to forward serialization of a Struct to a last resort net serializer. Optionally pass applicable UE_NET_LASTRESORT_NETSERIALIZER_OPTION_ options.
- * Example: UE_NET_IMPLEMENT_NAMED_STRUCT_LASTRESORT_NETSERIALIZER_AND_REGISTRY_DELEGATES(MyStruct, UE_NET_LASTRESORT_NETSERIALIZER_OPTION_MAXSIZE(32000), UE_NET_LASTRESORT_NETSERIALIZER_OPTION_EXCLUDEFROMDEFAULTSTATEHASH(true))
+ * Utility that can be used to forward serialization of a Struct to a last resort net serializer. Optionally pass FLastResortPropertyNetSerializerUserOptions.
+ * Example: UE_NET_IMPLEMENT_NAMED_STRUCT_LASTRESORT_NETSERIALIZER_AND_REGISTRY_DELEGATES(MyStruct, {.MaxQuantizedSizeBits = 32767, .CreateAndRegisterReplicationFragmentFunction = CreateAndRegisterMyStructReplicationFragment})
  */
-#define UE_NET_IMPLEMENT_NAMED_STRUCT_LASTRESORT_NETSERIALIZER_AND_REGISTRY_DELEGATES(Name, ... /* options */) \
+#define UE_NET_IMPLEMENT_NAMED_STRUCT_LASTRESORT_NETSERIALIZER_AND_REGISTRY_DELEGATES(Name, ... /* FLastResortPropertyNetSerializerUserOptions instance */) \
 	UE_NET_IS_DECLARED_TYPE(Name); \
 	static const FName PropertyNetSerializerRegistry_NAME_##Name( UE_STRINGIZE(Name) ); \
 	UE_NET_IMPLEMENT_NAMED_STRUCT_LASTRESORT_NETSERIALIZER_INFO(PropertyNetSerializerRegistry_NAME_##Name __VA_OPT__(,) __VA_ARGS__) \
