Fix double delete in ReplicationReader
Call GetReplicatedObjectInfo after FlushPostDispatchForBatch as it could modify ReplicatedObjects

#rnx

#Synced-CL 50524817

#Synced-CL 50524817

#Synced-CL 50524817

diff --git a/Engine/Source/Runtime/Net/Iris/Private/Iris/ReplicationSystem/ReplicationReader.cpp b/Engine/Source/Runtime/Net/Iris/Private/Iris/ReplicationSystem/ReplicationReader.cpp
--- a/Engine/Source/Runtime/Net/Iris/Private/Iris/ReplicationSystem/ReplicationReader.cpp
+++ b/Engine/Source/Runtime/Net/Iris/Private/Iris/ReplicationSystem/ReplicationReader.cpp
@@ -2071,8 +2071,6 @@ void FReplicationReader::DispatchStateData(FNetSerializationContext& Context)
 	// Dispatch and apply received state data
 	for (FDispatchObjectInfo& Info : ObjectsToDispatchArray->GetObjectsToDispatch())
 	{
-		FReplicatedObjectInfo* ReplicationInfo = GetReplicatedObjectInfo(Info.InternalIndex);
-
 		const FNetRefHandleManager::FReplicatedObjectData& ObjectData = NetRefHandleManager->GetReplicatedObjectDataNoCheck(Info.InternalIndex);
 
 		// Before starting a potentially new batch we want to flush rpc's and legacy callbacks belonging to the previous batch
@@ -2082,6 +2080,10 @@ void FReplicationReader::DispatchStateData(FNetSerializationContext& Context)
 			FlushPostDispatchForBatch();
 		}
 		LastDispatchedRootInternalIndex = RootInternalIndex;
+		
+		// This must be done after FlushPostDispatchForBatch as ReplicatedObjects could be modified
+		FReplicatedObjectInfo* ReplicationInfo = GetReplicatedObjectInfo(Info.InternalIndex);
+		CA_ASSUME(ReplicationInfo != nullptr);
 
 		FPostDispatchObjectInfo PostDispatchObjectInfo;
 		PostDispatchObjectInfo.ReplicationInfo = ReplicationInfo;
