Net/Iris: Add runtime filter creation

[RN] [net/Iris] Added support for runtime filter creation.
#Synced-CL 50832436

diff --git a/Engine/Source/Runtime/Net/Iris/Private/Iris/ReplicationSystem/Filtering/NoOpNetObjectFilter.cpp b/Engine/Source/Runtime/Net/Iris/Private/Iris/ReplicationSystem/Filtering/NoOpNetObjectFilter.cpp
new file mode 100644
--- /dev/null
+++ b/Engine/Source/Runtime/Net/Iris/Private/Iris/ReplicationSystem/Filtering/NoOpNetObjectFilter.cpp
@@ -0,0 +1,5 @@
+// Copyright Epic Games, Inc. All Rights Reserved.
+
+#include "NoOpNetObjectFilter.h"
+
+#include UE_INLINE_GENERATED_CPP_BY_NAME(NoOpNetObjectFilter)
diff --git a/Engine/Source/Runtime/Net/Iris/Private/Iris/ReplicationSystem/Filtering/NoOpNetObjectFilter.h b/Engine/Source/Runtime/Net/Iris/Private/Iris/ReplicationSystem/Filtering/NoOpNetObjectFilter.h
new file mode 100644
--- /dev/null
+++ b/Engine/Source/Runtime/Net/Iris/Private/Iris/ReplicationSystem/Filtering/NoOpNetObjectFilter.h
@@ -0,0 +1,23 @@
+// Copyright Epic Games, Inc. All Rights Reserved.
+
+#pragma once
+
+#include "Iris/ReplicationSystem/Filtering/NetObjectFilter.h"
+
+#include "NoOpNetObjectFilter.generated.h"
+
+// No-op filter. Used as a placeholder when a dynamic filter has been destroyed.
+UCLASS(Transient)
+class UNoOpNetObjectFilter : public UNetObjectFilter
+{
+	GENERATED_BODY()
+
+public:
+	virtual bool AddObject(uint32 ObjectIndex, FNetObjectFilterAddObjectParams&) override { return false; }
+	virtual void RemoveObject(uint32 ObjectIndex, const FNetObjectFilteringInfo&) override {}
+
+protected:
+	virtual void OnInit(const FNetObjectFilterInitParams&) override {}
+	virtual void OnDeinit() override {}
+	virtual void OnMaxInternalNetRefIndexIncreased(uint32 NewMaxInternalIndex) override {}
+};
diff --git a/Engine/Source/Runtime/Net/Iris/Private/Iris/ReplicationSystem/Filtering/ReplicationFiltering.cpp b/Engine/Source/Runtime/Net/Iris/Private/Iris/ReplicationSystem/Filtering/ReplicationFiltering.cpp
--- a/Engine/Source/Runtime/Net/Iris/Private/Iris/ReplicationSystem/Filtering/ReplicationFiltering.cpp
+++ b/Engine/Source/Runtime/Net/Iris/Private/Iris/ReplicationSystem/Filtering/ReplicationFiltering.cpp
@@ -1,6 +1,7 @@
 // Copyright Epic Games, Inc. All Rights Reserved.
 
 #include "ReplicationFiltering.h"
+#include "Algo/AnyOf.h"
 #include "HAL/PlatformMath.h"
 #include "HAL/IConsoleManager.h"
 #include "UObject/Package.h"
@@ -19,7 +20,10 @@
 #include "Iris/ReplicationSystem/ReplicationSystemInternal.h"
 #include "Iris/ReplicationSystem/Filtering/NetObjectGroups.h"
 #include "Iris/ReplicationSystem/Filtering/NetObjectFilterDefinitions.h"
+#include "Iris/ReplicationSystem/Filtering/NoOpNetObjectFilter.h"
 #include "Iris/ReplicationSystem/Filtering/ReplicationFilteringConfig.h"
+
+#include <inttypes.h>
 #include <limits>
 
 namespace UE::Net::Private
@@ -2881,43 +2885,119 @@ void FReplicationFiltering::InitFilters()
 
 	for (const FNetObjectFilterDefinition& FilterDefinition : FilterDefinitions->GetFilterDefinitions())
 	{
-		constexpr UObject* ClassOuter = nullptr;
-		UClass* NetObjectFilterClass = Cast<UClass>(StaticFindObject(UClass::StaticClass(), ClassOuter, ToCStr(FilterDefinition.ClassName.ToString()), EFindObjectFlags::ExactClass));
-		if (NetObjectFilterClass == nullptr || !NetObjectFilterClass->IsChildOf(UNetObjectFilter::StaticClass()))
+		CreateFilter(FilterDefinition);
+	}
+}
+
+TWeakObjectPtr<UNetObjectFilter> FReplicationFiltering::CreateFilter(const FNetObjectFilterDefinition& FilterDefinition)
+{
+	if (DynamicFilterInfos.Num() >= InvalidDynamicFilterIndex)
+	{
+		UE_LOG(
+			LogIrisFilterConfig,
+			Error,
+			TEXT("Cannot create filter %s: maximum number of filters reached (%" PRId32 ")"),
+			*FilterDefinition.FilterName.ToString(),
+			DynamicFilterInfos.Num());
+		return nullptr;
+	}
+
+	// Filter names must be unique.
+	if (Algo::AnyOf(DynamicFilterInfos, [&Name = FilterDefinition.FilterName](const FFilterInfo& Info) { return Info.Name == Name; }))
+	{
+		UE_LOG(
+			LogIrisFilterConfig,
+			Warning,
+			TEXT("Attempt to create filter %s when one with that name already exists"),
+			*FilterDefinition.FilterName.ToString());
+		return nullptr;
+	}
+
+	constexpr UObject* ClassOuter = nullptr;
+	UClass* NetObjectFilterClass = Cast<UClass>(StaticFindObject(UClass::StaticClass(), ClassOuter, ToCStr(FilterDefinition.ClassName.ToString()), EFindObjectFlags::ExactClass));
+	if (NetObjectFilterClass == nullptr || !NetObjectFilterClass->IsChildOf(UNetObjectFilter::StaticClass()))
+	{
+		ensureMsgf(NetObjectFilterClass != nullptr, TEXT("NetObjectFilter class is not a NetObjectFilter or could not be found: %s"), ToCStr(FilterDefinition.ClassName.ToString()));
+		return nullptr;
+	}
+
+	UClass* NetObjectFilterConfigClass = nullptr;
+	if (!FilterDefinition.ConfigClassName.IsNone())
+	{
+		NetObjectFilterConfigClass = Cast<UClass>(StaticFindObject(UClass::StaticClass(), ClassOuter, ToCStr(FilterDefinition.ConfigClassName.ToString()), EFindObjectFlags::ExactClass));
+		if (NetObjectFilterConfigClass == nullptr || !NetObjectFilterConfigClass->IsChildOf(UNetObjectFilterConfig::StaticClass()))
 		{
-			ensureMsgf(NetObjectFilterClass != nullptr, TEXT("NetObjectFilter class is not a NetObjectFilter or could not be found: %s"), ToCStr(FilterDefinition.ClassName.ToString()));
-			continue;
+			ensureMsgf(NetObjectFilterConfigClass != nullptr, TEXT("NetObjectFilterConfig class is not a NetObjectFilterConfig or could not be found: %s"), ToCStr(FilterDefinition.ConfigClassName.ToString()));
+			return nullptr;
 		}
+	}
+
+	FFilterInfo& Info = DynamicFilterInfos.Emplace_GetRef();
+	Info.Filter = TStrongObjectPtr<UNetObjectFilter>(NewObject<UNetObjectFilter>(GetTransientPackage(), NetObjectFilterClass, MakeUniqueObjectName(nullptr, NetObjectFilterClass, FilterDefinition.FilterName)));
+	check(Info.Filter.IsValid());
+	Info.Name = FilterDefinition.FilterName;
+	Info.ObjectCount = 0;
+
+	Info.Filter->Init(GetFilterInitParams(NetObjectFilterConfigClass));
+
+	bHasDynamicFilters = true;
+	bHasDynamicFiltersWithUpdateTrait = bHasDynamicFiltersWithUpdateTrait || Info.Filter->HasFilterTrait(ENetFilterTraits::NeedsUpdate);
+
+	ValidConnections.ForAllSetBits([&Filter = *Info.Filter](const uint32 ConnectionId) { Filter.AddConnection(ConnectionId); });
 
-		UClass* NetObjectFilterConfigClass = nullptr;
-		if (!FilterDefinition.ConfigClassName.IsNone())
+	return Info.Filter.Get();
+}
+
+void FReplicationFiltering::DestroyFilter(const FName Name)
+{
+	for (FFilterInfo& FilterInfo : DynamicFilterInfos)
+	{
+		if (FilterInfo.Name != Name || !FilterInfo.Filter)
 		{
-			NetObjectFilterConfigClass = Cast<UClass>(StaticFindObject(UClass::StaticClass(), ClassOuter, ToCStr(FilterDefinition.ConfigClassName.ToString()), EFindObjectFlags::ExactClass));
-			if (NetObjectFilterConfigClass == nullptr || !NetObjectFilterConfigClass->IsChildOf(UNetObjectFilterConfig::StaticClass()))
-			{
-				ensureMsgf(NetObjectFilterConfigClass != nullptr, TEXT("NetObjectFilterConfig class is not a NetObjectFilterConfig or could not be found: %s"), ToCStr(FilterDefinition.ConfigClassName.ToString()));
-				continue;
-			}
+			continue;
 		}
 
-		
-		FFilterInfo& Info = DynamicFilterInfos.Emplace_GetRef();
-		Info.Filter = TStrongObjectPtr<UNetObjectFilter>(NewObject<UNetObjectFilter>(GetTransientPackage(), NetObjectFilterClass, MakeUniqueObjectName(nullptr, NetObjectFilterClass, FilterDefinition.FilterName)));
-		check(Info.Filter.IsValid());
-		Info.Name = FilterDefinition.FilterName;
-		Info.ObjectCount = 0;
+		UE_LOG(LogIrisFiltering, Verbose, TEXT("Destroying filter \"%s\""), *FilterInfo.Name.ToString());
+
+		UNetObjectFilter& Filter = *FilterInfo.Filter;
+
+		// Unassign any objects assigned to the filter.
+		// Construct FNetBitArray from FNetBitArrayView (this conversion is not provided by FNetBitArray at time of writing).
+		auto CloneBitArray = [](const FNetBitArrayView& View) -> FNetBitArray
+			{
+				FNetBitArray BitArray;
+				BitArray.InitAndCopy(View);
+				return BitArray;
+			};
+
+		// Local copy to avoid modifying the bitset while iterating it
+		if (const FNetBitArray AssignedObjects = CloneBitArray(Filter.GetFilteredObjects());
+			AssignedObjects.IsAnyBitSet())
+		{
+			UE_LOG(
+				LogIrisFiltering,
+				Log,
+				TEXT("Orphaning %" PRIu32 " objects assigned to filter %s at filter destruction"),
+				AssignedObjects.CountSetBits(),
+				*Filter.GetName());
+
+			AssignedObjects.ForAllSetBits(
+				[this](uint32 ObjectIndex)
+				{
+					SetFilter(FInternalNetRefIndex(ObjectIndex), InvalidNetObjectFilterHandle, NAME_None);
+				});
+		}
 
-		FNetObjectFilterInitParams InitParams;
-		InitParams.ReplicationSystem = ReplicationSystem;
-		InitParams.Config = (NetObjectFilterConfigClass ? NewObject<UNetObjectFilterConfig>(GetTransientPackage(), NetObjectFilterConfigClass) : nullptr);
-		InitParams.AbsoluteMaxNetObjectCount = NetRefHandleManager->GetMaxActiveObjectCount();
-		InitParams.CurrentMaxInternalIndex = MaxInternalNetRefIndex;
-		InitParams.MaxConnectionCount = Connections->GetMaxConnectionCount();
+		Filter.Deinit();
 
-		Info.Filter->Init(InitParams);
+		// Destroy the filter but keep its index reserved to avoid renumbering subsequent filter handles and prevent reuse of stale handles with
+		// an unintended filter.
+		FilterInfo = FFilterInfo{
+			.Filter = MakeNoOpNetObjectFilter(),
+			.Name = NAME_None,
+		};
 
-		bHasDynamicFilters = true;
-		bHasDynamicFiltersWithUpdateTrait = bHasDynamicFiltersWithUpdateTrait || Info.Filter->HasFilterTrait(ENetFilterTraits::NeedsUpdate);
+		break; // Filter names are unique; only one may match.
 	}
 }
 
@@ -3148,6 +3228,31 @@ bool FReplicationFiltering::HasSubObjectInScopeWithFilteredOutRootObject(uint32
 	return false;
 }
 
+FNetObjectFilterInitParams FReplicationFiltering::GetFilterInitParams(UClass* ConfigClass) const
+{
+	return FNetObjectFilterInitParams{
+		.ReplicationSystem = ReplicationSystem,
+		.Config = ConfigClass ? NewObject<UNetObjectFilterConfig>(GetTransientPackage(), ConfigClass) : nullptr,
+		.AbsoluteMaxNetObjectCount = NetRefHandleManager->GetMaxActiveObjectCount(),
+		.CurrentMaxInternalIndex = MaxInternalNetRefIndex,
+		.MaxConnectionCount = Connections->GetMaxConnectionCount(),
+	};
+}
+
+TStrongObjectPtr<UNetObjectFilter> FReplicationFiltering::MakeNoOpNetObjectFilter() const
+{
+	TStrongObjectPtr<UNetObjectFilter> Filter = TStrongObjectPtr<UNetObjectFilter>(NewObject<UNoOpNetObjectFilter>());
+
+	// Initialize the filter essentially to satisfy code analysis tools.
+	Filter->Init(GetFilterInitParams(/* ConfigClass */ nullptr));
+	ValidConnections.ForAllSetBits([&Filter = *Filter](const uint32 ConnectionId)
+		{
+			Filter.AddConnection(ConnectionId);
+		});
+
+	return Filter;
+}
+
 void FReplicationFiltering::FPerConnectionInfo::Deinit()
 {
 	ConnectionFilteredObjects.Empty();
diff --git a/Engine/Source/Runtime/Net/Iris/Private/Iris/ReplicationSystem/Filtering/ReplicationFiltering.h b/Engine/Source/Runtime/Net/Iris/Private/Iris/ReplicationSystem/Filtering/ReplicationFiltering.h
--- a/Engine/Source/Runtime/Net/Iris/Private/Iris/ReplicationSystem/Filtering/ReplicationFiltering.h
+++ b/Engine/Source/Runtime/Net/Iris/Private/Iris/ReplicationSystem/Filtering/ReplicationFiltering.h
@@ -12,9 +12,12 @@
 #include "Net/Core/Connection/ConnectionHandle.h"
 #include "Containers/Array.h"
 #include "UObject/StrongObjectPtr.h"
+#include "UObject/WeakObjectPtrTemplates.h"
 
 class UReplicationFilteringConfig;
 class UReplicationSystem;
+struct FNetObjectFilterDefinition;
+
 namespace UE::Net
 {
 	typedef uint32 FNetObjectFilterHandle;
@@ -108,6 +111,18 @@ public:
 	/** Returns the name of the Filter represented by the handle. */
 	FName GetFilterName(FNetObjectFilterHandle Filter) const;
 
+	/**
+	 * Create a filter at runtime.
+	 * @see UReplicationSystem::CreateFilter().
+	 */
+	TWeakObjectPtr<UNetObjectFilter> CreateFilter(const FNetObjectFilterDefinition& Definition);
+
+	/**
+	* Destroy a filter.
+	* @see UReplicationSystem::DestroyFilter().
+	*/
+	void DestroyFilter(const FName FilterName);
+
 	/** Fill the passed in list with root objects considered always relevant */
 	void BuildAlwaysRelevantList(FNetBitArrayView OutAlwaysRelevantList, const FNetBitArrayView ScopeList) const;
 
@@ -302,6 +317,14 @@ private:
 	bool HasSubObjectInScopeWithFilteredOutRootObject(FNetBitArrayView Objects) const;
 	bool HasSubObjectInScopeWithFilteredOutRootObject(uint32 connectionId) const;
 
+	// Get InitParams for a filter.
+	// @param ConfigClass Filter-specific config class. May be nullptr.
+	// @return InitParams suitable for passing to UNetObjectFilter::Init().
+	FNetObjectFilterInitParams GetFilterInitParams(UClass* ConfigClass) const;
+
+	// Create & initialize an instance of the no-op filter.
+	TStrongObjectPtr<UNetObjectFilter> MakeNoOpNetObjectFilter() const;
+
 private:
 	enum EHysteresisProcessingMode : uint32
 	{
diff --git a/Engine/Source/Runtime/Net/Iris/Private/Iris/ReplicationSystem/ObjectReplicationBridge.cpp b/Engine/Source/Runtime/Net/Iris/Private/Iris/ReplicationSystem/ObjectReplicationBridge.cpp
--- a/Engine/Source/Runtime/Net/Iris/Private/Iris/ReplicationSystem/ObjectReplicationBridge.cpp
+++ b/Engine/Source/Runtime/Net/Iris/Private/Iris/ReplicationSystem/ObjectReplicationBridge.cpp
@@ -3176,6 +3176,11 @@ UE::Net::FNetObjectPrioritizerHandle UObjectReplicationBridge::GetPrioritizer(co
 	return InvalidNetObjectPrioritizerHandle;
 }
 
+void UObjectReplicationBridge::ReloadConfig()
+{
+	LoadConfig();
+}
+
 void UObjectReplicationBridge::LoadConfig()
 {
 	// Clear everything related to the config.
diff --git a/Engine/Source/Runtime/Net/Iris/Private/Iris/ReplicationSystem/ReplicationSystem.cpp b/Engine/Source/Runtime/Net/Iris/Private/Iris/ReplicationSystem/ReplicationSystem.cpp
--- a/Engine/Source/Runtime/Net/Iris/Private/Iris/ReplicationSystem/ReplicationSystem.cpp
+++ b/Engine/Source/Runtime/Net/Iris/Private/Iris/ReplicationSystem/ReplicationSystem.cpp
@@ -27,6 +27,7 @@
 #include "Iris/ReplicationSystem/ObjectReferenceCache.h"
 
 #include "Iris/ReplicationSystem/Filtering/NetObjectFilter.h"
+#include "Iris/ReplicationSystem/Filtering/NetObjectFilterDefinitions.h"
 #include "Iris/ReplicationSystem/Filtering/ReplicationFiltering.h"
 
 #include "Iris/DataStream/DataStreamManager.h"
@@ -657,6 +658,33 @@ public:
 		}
 	}
 
+	TWeakObjectPtr<UNetObjectFilter> CreateFilter(const FNetObjectFilterDefinition& Definition)
+	{
+		if (!ensure(!ReplicationSystemInternal.AreFilterChangesBlocked()))
+		{
+			UE_LOG(
+				LogIris,
+				Error,
+				TEXT("CreateFilter() called while filter changes are blocked; filter %s not created"),
+				*Definition.FilterName.ToString());
+
+			return nullptr;
+		}
+
+		return ReplicationSystemInternal.GetFiltering().CreateFilter(Definition);
+	}
+
+	void DestroyFilter(const FName Name)
+	{
+		if (ReplicationSystemInternal.AreFilterChangesBlocked())
+		{
+			UE_LOG(LogIris, Warning, TEXT("DestroyFilter() called while filter changes are blocked; filter %s not destroyed"), *Name.ToString());
+			return;
+		}
+
+		ReplicationSystemInternal.GetFiltering().DestroyFilter(Name);
+	}
+
 	void AddConnection(uint32 ConnectionId)
 	{
 		LLM_SCOPE_BYTAG(IrisConnection);
@@ -1572,6 +1600,16 @@ const UE::Net::FNetDebugName* UReplicationSystem::GetDebugName(FNetRefHandle Han
 	return Protocol ? Protocol->DebugName : nullptr;
 }
 
+TWeakObjectPtr<UNetObjectFilter> UReplicationSystem::CreateFilter(const FNetObjectFilterDefinition& Definition)
+{
+	return Impl->CreateFilter(Definition);
+}
+
+void UReplicationSystem::DestroyFilter(const FName Name)
+{
+	Impl->DestroyFilter(Name);
+}
+
 void UReplicationSystem::SetOwningNetConnection(FNetRefHandle Handle, uint32 ConnectionId)
 {
 	using namespace UE::Net::Private;
diff --git a/Engine/Source/Runtime/Net/Iris/Public/Iris/ReplicationSystem/Filtering/NetObjectFilterDefinitions.h b/Engine/Source/Runtime/Net/Iris/Public/Iris/ReplicationSystem/Filtering/NetObjectFilterDefinitions.h
--- a/Engine/Source/Runtime/Net/Iris/Public/Iris/ReplicationSystem/Filtering/NetObjectFilterDefinitions.h
+++ b/Engine/Source/Runtime/Net/Iris/Public/Iris/ReplicationSystem/Filtering/NetObjectFilterDefinitions.h
@@ -12,14 +12,14 @@ struct FNetObjectFilterDefinition
 	GENERATED_BODY()
 
 	/** Filter identifier. Used to uniquely identify a filter. */
-	UPROPERTY()
+	UPROPERTY(EditAnywhere, Category="Filter")
 	FName FilterName;
 
 	/**
 	 * UClass name, specified by its fully qualified path, used to create the UNetObjectFilter. You can have multiple instances
 	 * of the same filter as long as their FilterNames are unique.
 	 */
-	UPROPERTY()
+	UPROPERTY(EditAnywhere, Category="Filter")
 	FName ClassName;		
 
 	/**
@@ -27,7 +27,7 @@ struct FNetObjectFilterDefinition
 	 * will be passed at filter initialization time. If you want multiple instances of the same
 	 * filter then use subclassing to create unique filter configs.
 	 */
-	UPROPERTY()
+	UPROPERTY(EditAnywhere, Category="Filter")
 	FName ConfigClassName;		
 };
 
diff --git a/Engine/Source/Runtime/Net/Iris/Public/Iris/ReplicationSystem/ObjectReplicationBridge.h b/Engine/Source/Runtime/Net/Iris/Public/Iris/ReplicationSystem/ObjectReplicationBridge.h
--- a/Engine/Source/Runtime/Net/Iris/Public/Iris/ReplicationSystem/ObjectReplicationBridge.h
+++ b/Engine/Source/Runtime/Net/Iris/Public/Iris/ReplicationSystem/ObjectReplicationBridge.h
@@ -288,6 +288,9 @@ public:
 	 */
 	IRISCORE_API void PreRegisterObjectWithReferenceHandle(const UObject* Instance, FNetRefHandle Handle);
 
+	/** Re-initialize config-driven parameters found in ObjectReplicationBridgeConfig. */
+	IRISCORE_API void ReloadConfig();
+
 public:
 
 	struct FRootObjectReplicationParams
diff --git a/Engine/Source/Runtime/Net/Iris/Public/Iris/ReplicationSystem/ReplicationSystem.h b/Engine/Source/Runtime/Net/Iris/Public/Iris/ReplicationSystem/ReplicationSystem.h
--- a/Engine/Source/Runtime/Net/Iris/Public/Iris/ReplicationSystem/ReplicationSystem.h
+++ b/Engine/Source/Runtime/Net/Iris/Public/Iris/ReplicationSystem/ReplicationSystem.h
@@ -31,6 +31,9 @@ class UNetObjectFilter;
 class UNetObjectPrioritizer;
 class UObjectReplicationBridge;
 class UReplicationSystem;
+
+struct FNetObjectFilterDefinition;
+
 namespace UE::Net
 {
 	typedef uint32 FNetObjectFilterHandle;
@@ -553,6 +556,36 @@ public:
 
 	// Filtering
 
+	/**
+	 * Create a filter at runtime.
+	 * 
+	 * Definition.FilterName must be unique.
+	 *
+	 * Rather than explicitly creating filters at runtime, they can be specified in config files and instantiated automatically by the network
+	 * system. See https://dev.epicgames.com/documentation/en-us/unreal-engine/iris-filtering-in-unreal-engine#dynamic-filter-configuration .
+	 *
+	 * If .ini configuration applies to the filter or its use, call UObjectReplicationBridge::ReloadConfig() after creating the filter to apply that
+	 * config.
+	 *
+	 * The filter is owned by the replication system, but can be accessed via the returned weak pointer.
+	 *
+	 * The filter can be destroyed later via DestroyFilter().
+	 *
+	 * @param Definition Filter definition (class, config).
+	 *
+	 * @return Weak pointer to filter instance, or nullptr.
+	 */
+	IRISCORE_API TWeakObjectPtr<UNetObjectFilter> CreateFilter(const FNetObjectFilterDefinition& Definition);
+
+	/**
+	 * Destroy a filter.
+	 *
+	 * Should only be used on filters created at runtime by calls to CreateFilter().
+	 *
+	 * Any objects still assigned to the filter are unassigned.
+	 */
+	IRISCORE_API void DestroyFilter(const FName Name);
+
 	/**
 	 * Sets the owning connection of an object. This can determine which properties are replicated 
 	 * to which connections or affect filtering such that the object is only replicated to the owning connection.
