Iris - Allow users to subclass UNetObjectPrioritizer easily.

#Synced-CL 50723010

diff --git a/Engine/Source/Runtime/Net/Iris/Public/Iris/ReplicationSystem/Prioritization/NetObjectPrioritizer.h b/Engine/Source/Runtime/Net/Iris/Public/Iris/ReplicationSystem/Prioritization/NetObjectPrioritizer.h
--- a/Engine/Source/Runtime/Net/Iris/Public/Iris/ReplicationSystem/Prioritization/NetObjectPrioritizer.h
+++ b/Engine/Source/Runtime/Net/Iris/Public/Iris/ReplicationSystem/Prioritization/NetObjectPrioritizer.h
@@ -153,7 +153,7 @@ struct FNetObjectPrioritizerUpdateParams
  * object to still not be replicated to a particular connection a certain frame. There is no mechanism to force an object to be replicated a 
  * certain frame, but the priority is a major factor in the decision.
  */
-UCLASS(Abstract)
+UCLASS(Abstract, MinimalAPI)
 class UNetObjectPrioritizer : public UObject
 {
 	GENERATED_BODY()
