Force static object references to serialize as paths in backend Proxy RPCs to ensure consistent client-assigned reference behavior.

#Synced-CL 50902502

diff --git a/Engine/Source/Runtime/Net/Iris/Private/Iris/ReplicationSystem/NetBlob/NetRPC.cpp b/Engine/Source/Runtime/Net/Iris/Private/Iris/ReplicationSystem/NetBlob/NetRPC.cpp
--- a/Engine/Source/Runtime/Net/Iris/Private/Iris/ReplicationSystem/NetBlob/NetRPC.cpp
+++ b/Engine/Source/Runtime/Net/Iris/Private/Iris/ReplicationSystem/NetBlob/NetRPC.cpp
@@ -421,6 +421,16 @@ FNetRPC* FNetRPC::Create(UReplicationSystem* ReplicationSystem, const FNetBlobCr
 		FInternalNetSerializationContext InternalContext(ReplicationSystem);
 		
 		Context.SetInternalContext(&InternalContext);
+		
+		// The proxy uses a single, shared object reference cache and can contain references from multiple game servers, which means it can differ from
+		// the object reference caches of the individual connected clients. For this reason, any RPC arguments forwarded from a client to a game server
+		// through the proxy that are static object references should always use objects paths. This avoids the proxy generating a new object reference
+		// using a FNetRefHandle from the proxy's object reference cache when it doesn't exist on the client (the game server expects it's clients to have
+		// an object reference cache that is a subset of it's own cache).
+		if (ReplicationSystem->GetProxyType() == EProxyType::Backend)
+		{
+			InternalContext.bForceObjectPathSerialization = 1;
+		}
 
 		// Quantize the function parameters
 		FReplicationStateOperations::Quantize(Context, QuantizedBlobState.GetStateBuffer(), static_cast<const uint8*>(FunctionParameters), BlobDescriptor);
diff --git a/Engine/Source/Runtime/Net/Iris/Private/Iris/ReplicationSystem/ObjectReferenceCache.cpp b/Engine/Source/Runtime/Net/Iris/Private/Iris/ReplicationSystem/ObjectReferenceCache.cpp
--- a/Engine/Source/Runtime/Net/Iris/Private/Iris/ReplicationSystem/ObjectReferenceCache.cpp
+++ b/Engine/Source/Runtime/Net/Iris/Private/Iris/ReplicationSystem/ObjectReferenceCache.cpp
@@ -247,6 +247,21 @@ bool FObjectReferenceCache::RenamePathForPie(uint32 ConnectionId, FString& Str,
 	}
 }
 
+FNetObjectReference FObjectReferenceCache::GetOrCreateObjectReferenceUsingPath(const UObject* Instance, UReplicationSystem* ReplicationSystem)
+{
+	if (Instance && Instance->IsFullNameStableForNetworking())
+	{
+		FStringTokenStore* LocalStringTokenStore = ReplicationSystem->GetNetTokenStore()->GetDataStore<FStringTokenStore>();
+
+		FString ObjectPath = Instance->GetPathName();
+		const FNetToken ObjectPathToken = LocalStringTokenStore->GetOrCreateToken(ObjectPath);
+
+		return MakeNetObjectReference(FNetRefHandle::GetInvalid(), ObjectPathToken);
+	}
+
+	return GetOrCreateObjectReference(Instance, ReplicationSystem);
+}
+
 FNetRefHandle FObjectReferenceCache::CreateObjectReferenceHandle(const UObject* Object, UReplicationSystem* ReplicationSystem)
 {
 	FNetObjectReference Ref;
diff --git a/Engine/Source/Runtime/Net/Iris/Private/Iris/ReplicationSystem/ObjectReferenceCache.h b/Engine/Source/Runtime/Net/Iris/Private/Iris/ReplicationSystem/ObjectReferenceCache.h
--- a/Engine/Source/Runtime/Net/Iris/Private/Iris/ReplicationSystem/ObjectReferenceCache.h
+++ b/Engine/Source/Runtime/Net/Iris/Private/Iris/ReplicationSystem/ObjectReferenceCache.h
@@ -49,6 +49,8 @@ struct FObjectReferenceCacheParams
 
 class FObjectReferenceCache
 {
+	template<typename T> friend struct FObjectNetSerializerBase;
+
 public:
 
 	FObjectReferenceCache();
@@ -241,6 +243,9 @@ private:
 	bool CanClientLoadObjectInternal(const UObject* Object, bool bIsDynamic) const;
 	bool ShouldIgnoreWhenMissing(FNetRefHandle RefHandle) const;
 	bool RenamePathForPie(uint32 ConnectionId, FString& Str, bool bReading, UReplicationSystem* ReplicationSystem);
+	
+	// Get or create a NetObjectReference from an object: always using the object path if static, and a NetRefHandle if dynamic.
+	FNetObjectReference GetOrCreateObjectReferenceUsingPath(const UObject* Instance, UReplicationSystem* ReplicationSystem);
 
 	// Get the string path of RefHandle
 	FString FullPath(FNetRefHandle RefHandle, const FNetObjectResolveContext& ResolveContext) const;
diff --git a/Engine/Source/Runtime/Net/Iris/Private/Iris/Serialization/InternalNetSerializationContext.h b/Engine/Source/Runtime/Net/Iris/Private/Iris/Serialization/InternalNetSerializationContext.h
--- a/Engine/Source/Runtime/Net/Iris/Private/Iris/Serialization/InternalNetSerializationContext.h
+++ b/Engine/Source/Runtime/Net/Iris/Private/Iris/Serialization/InternalNetSerializationContext.h
@@ -61,6 +61,9 @@ public:
 
 	// If true, allow calling rep notifies
 	uint32 bAllowExecuteRepNotify : 1;
+
+	// If true, always quantize a static object reference using the object's path
+	uint32 bForceObjectPathSerialization : 1 = 0;
 };
 
 inline FInternalNetSerializationContext::FInternalNetSerializationContext()
diff --git a/Engine/Source/Runtime/Net/Iris/Private/Iris/Serialization/ObjectNetSerializer.cpp b/Engine/Source/Runtime/Net/Iris/Private/Iris/Serialization/ObjectNetSerializer.cpp
--- a/Engine/Source/Runtime/Net/Iris/Private/Iris/Serialization/ObjectNetSerializer.cpp
+++ b/Engine/Source/Runtime/Net/Iris/Private/Iris/Serialization/ObjectNetSerializer.cpp
@@ -336,7 +336,21 @@ void FObjectNetSerializerBase<T>::Quantize(FNetSerializationContext& Context, co
 	{
 		UObject* SourceObject = GetRawPtrValue(Source);
 		UReplicationSystem* ReplicationSystem = InternalContext->ReplicationSystem;
-		Target.SetNetReference(Context, Args, Context.IsInitializingDefaultState() ? FNetObjectReference() : InternalContext->ObjectReferenceCache->GetOrCreateObjectReference(SourceObject, ReplicationSystem));
+
+		FNetObjectReference Reference;
+		if (!Context.IsInitializingDefaultState())
+		{
+			if (InternalContext->bForceObjectPathSerialization)
+			{
+				Reference = InternalContext->ObjectReferenceCache->GetOrCreateObjectReferenceUsingPath(SourceObject, ReplicationSystem);
+			}
+			else
+			{
+				Reference = InternalContext->ObjectReferenceCache->GetOrCreateObjectReference(SourceObject, ReplicationSystem);
+			}
+		}
+
+		Target.SetNetReference(Context, Args, Reference);
 	}
 }
 
