Set ObjectRecord.NextIndex to invalid index on creation to mitigate infinite loop in FReplicationWriter::HasInFlightStateChanges

diff --git a/Engine/Source/Runtime/Net/Iris/Private/Iris/ReplicationSystem/ReplicationWriter.cpp b/Engine/Source/Runtime/Net/Iris/Private/Iris/ReplicationSystem/ReplicationWriter.cpp
--- a/Engine/Source/Runtime/Net/Iris/Private/Iris/ReplicationSystem/ReplicationWriter.cpp
+++ b/Engine/Source/Runtime/Net/Iris/Private/Iris/ReplicationSystem/ReplicationWriter.cpp
@@ -2004,6 +2004,8 @@ void FReplicationWriter::CreateObjectRecord(const FNetBitArrayView* ChangeMask,
 
 	RecordInfo.Index = ObjectInfo.InternalIndex;
 
+	RecordInfo.NextIndex = FReplicationRecord::InvalidReplicationRecordIndex;
+
 	if (ObjectInfo.AttachmentType == ENetObjectAttachmentType::HugeObject)
 	{
 		RecordInfo.ReplicatedObjectState = uint8(EReplicatedObjectState::HugeObject);
