Iris
* Fix ensure that didn't fire if the filter config wasn't the right class type

#rnx

#Synced-CL 51003732

diff --git a/Engine/Source/Runtime/Net/Iris/Private/Iris/ReplicationSystem/Filtering/ReplicationFiltering.cpp b/Engine/Source/Runtime/Net/Iris/Private/Iris/ReplicationSystem/Filtering/ReplicationFiltering.cpp
--- a/Engine/Source/Runtime/Net/Iris/Private/Iris/ReplicationSystem/Filtering/ReplicationFiltering.cpp
+++ b/Engine/Source/Runtime/Net/Iris/Private/Iris/ReplicationSystem/Filtering/ReplicationFiltering.cpp
@@ -1509,17 +1509,21 @@ void FReplicationFiltering::UpdateDynamicFiltering()
 	const uint32* SubObjectsData = NetRefHandleManager->GetSubObjectInternalIndices().GetDataChecked(WordCountForObjectBitArrays);
 	const uint32* ObjectsRequiringDynamicFilterUpdateData = ObjectsRequiringDynamicFilterUpdate.GetDataChecked(WordCountForObjectBitArrays);
 
+	// List of rootobjects that have a dependency trait with another parent root object
 	const FNetBitArrayView DependentObjectsView = NetRefHandleManager->GetDependentObjectInternalIndices();
 	const uint32* DependentObjectsData = DependentObjectsView.GetDataChecked(WordCountForObjectBitArrays);
 
-	FNetBitArrayView ObjectsWithDirtyCreationDependencies = NetRefHandleManager->GetObjectsWithDirtyCreationDependencies();
+	// List of objects that had a change in their creation dependency
+	const FNetBitArrayView ObjectsWithDirtyCreationDependencies = NetRefHandleManager->GetObjectsWithDirtyCreationDependencies();
 
+	// List of rootobjects that have childs depending on them
 	const FNetBitArrayView ObjectsWithCreationDependents = NetRefHandleManager->GetObjectsWithCreationDependents();
 	const uint32* ObjectsWithCreationDependentsData = ObjectsWithCreationDependents.GetDataChecked(WordCountForObjectBitArrays);
 
+	// List of child rootobjects that depend on a parent for instantiation
 	const FNetBitArrayView ObjectsWithCreationDependencies = NetRefHandleManager->GetObjectsWithCreationDependencies();
 	const uint32* ObjectsWithCreationDependenciesData = ObjectsWithCreationDependencies.GetDataChecked(WordCountForObjectBitArrays);
-	
+
 	uint32* ConnectionIds = static_cast<uint32*>(FMemory_Alloca(ValidConnections.GetNumBits() * sizeof(uint32)));
 	uint32 ConnectionCount = 0;
 	ValidConnections.ForAllSetBits([ConnectionIds, &ConnectionCount](uint32 Bit) { ConnectionIds[ConnectionCount++] = Bit; });
@@ -1548,13 +1552,15 @@ void FReplicationFiltering::UpdateDynamicFiltering()
 					continue;
 				}
 
-				FNetObjectFilteringParams FilteringParams{
+				FNetObjectFilteringParams FilteringParams
+				{
 					.OutAllowedObjects = AllowedObjects,
 					.FilteringInfos = MakeArrayView(NetObjectFilteringInfos),
 					.StateBuffers = nullptr,
 					.ConnectionId = ConnId,
 					.View = Connections->GetReplicationView(ConnId),
-					.GroupFilteredOutObjects = ReplicationSystem->GetReplicationSystemInternal()->GetFiltering().GetGroupFilteredOutObjects(ConnId)};
+					.GroupFilteredOutObjects = ReplicationSystem->GetReplicationSystemInternal()->GetFiltering().GetGroupFilteredOutObjects(ConnId)
+				};
 
 				// Execute the filter here
 				Info.Filter->Filter(FilteringParams);
@@ -1754,6 +1760,7 @@ void FReplicationFiltering::UpdateDynamicFiltering()
 
 		ensure(!bCVarRepFilterValidateNoSubObjectInScopeWithFilteredOutRootObject || !HasSubObjectInScopeWithFilteredOutRootObject(DynamicFilteredOutObjectsHysteresisAdjusted));
 
+		// New optimized creation dependency path
 		if (bCVarRepFilterUpdateCreationDependenciesOnRelevancyChange)
 		{
 			IRIS_PROFILER_SCOPE(FReplicationFiltering_UpdateCreationDependencies);
@@ -1823,6 +1830,7 @@ void FReplicationFiltering::UpdateDynamicFiltering()
 		}
 
 		// Update dependent objects that can only be relevant if their parent also is
+		// Old Non-optimized creation dependency path
 		if (!bCVarRepFilterUpdateCreationDependenciesOnRelevancyChange)
 		{
 			IRIS_PROFILER_SCOPE(FReplicationFiltering_UpdateCreationDependencies);
@@ -1831,7 +1839,8 @@ void FReplicationFiltering::UpdateDynamicFiltering()
 				FNetBitArrayView ConnectionObjectsInScope = MakeNetBitArrayView(ConnectionInfo.ObjectsInScope);
 				FNetBitArray ProcessedObjects(ConnectionInfo.ObjectsInScope.GetNumBits());
 				FNetBitArrayView ProcessedObjectsView = MakeNetBitArrayView(ProcessedObjects);
-				FNetBitArrayView::ForAllSetBits(ObjectsWithCreationDependencies, ConnectionObjectsInScope, FNetBitArrayView::AndOp, [ & ](uint32 ChildIndex)
+
+				FNetBitArrayView::ForAllSetBits(ObjectsWithCreationDependencies, ConnectionObjectsInScope, FNetBitArrayView::AndOp, [&](uint32 ChildIndex)
 				{
 					constexpr bool bIsNotRecursive = false;
 					this->UpdateCreationDependentParent(ChildIndex, ObjectsWithCreationDependencies, ConnectionObjectsInScope, ProcessedObjectsView, bIsNotRecursive);
@@ -2951,7 +2960,7 @@ TWeakObjectPtr<UNetObjectFilter> FReplicationFiltering::CreateFilter(const FNetO
 	UClass* NetObjectFilterClass = Cast<UClass>(StaticFindObject(UClass::StaticClass(), ClassOuter, ToCStr(FilterDefinition.ClassName.ToString()), EFindObjectFlags::ExactClass));
 	if (NetObjectFilterClass == nullptr || !NetObjectFilterClass->IsChildOf(UNetObjectFilter::StaticClass()))
 	{
-		ensureMsgf(NetObjectFilterClass != nullptr, TEXT("NetObjectFilter class is not a NetObjectFilter or could not be found: %s"), ToCStr(FilterDefinition.ClassName.ToString()));
+		ensureMsgf(false, TEXT("NetObjectFilter class is not a NetObjectFilter or could not be found: %s"), ToCStr(FilterDefinition.ClassName.ToString()));
 		return nullptr;
 	}
 
@@ -2961,7 +2970,7 @@ TWeakObjectPtr<UNetObjectFilter> FReplicationFiltering::CreateFilter(const FNetO
 		NetObjectFilterConfigClass = Cast<UClass>(StaticFindObject(UClass::StaticClass(), ClassOuter, ToCStr(FilterDefinition.ConfigClassName.ToString()), EFindObjectFlags::ExactClass));
 		if (NetObjectFilterConfigClass == nullptr || !NetObjectFilterConfigClass->IsChildOf(UNetObjectFilterConfig::StaticClass()))
 		{
-			ensureMsgf(NetObjectFilterConfigClass != nullptr, TEXT("NetObjectFilterConfig class is not a NetObjectFilterConfig or could not be found: %s"), ToCStr(FilterDefinition.ConfigClassName.ToString()));
+			ensureMsgf(false, TEXT("NetObjectFilterConfig class is not a NetObjectFilterConfig or could not be found: %s"), ToCStr(FilterDefinition.ConfigClassName.ToString()));
 			return nullptr;
 		}
 	}
