Iris Filtering: Skip nulled placeholder filters

Add explicit nullptr checks rather than using a NoOp placeholder filter.

#rnx

#Synced-CL 51103119

#Synced-CL 51103119

diff --git a/Engine/Source/Runtime/Net/Iris/Private/Iris/ReplicationSystem/Filtering/NoOpNetObjectFilter.cpp b/Engine/Source/Runtime/Net/Iris/Private/Iris/ReplicationSystem/Filtering/NoOpNetObjectFilter.cpp
deleted file mode 100644
diff --git a/Engine/Source/Runtime/Net/Iris/Private/Iris/ReplicationSystem/Filtering/NoOpNetObjectFilter.h b/Engine/Source/Runtime/Net/Iris/Private/Iris/ReplicationSystem/Filtering/NoOpNetObjectFilter.h
deleted file mode 100644
diff --git a/Engine/Source/Runtime/Net/Iris/Private/Iris/ReplicationSystem/Filtering/ReplicationFiltering.cpp b/Engine/Source/Runtime/Net/Iris/Private/Iris/ReplicationSystem/Filtering/ReplicationFiltering.cpp
--- a/Engine/Source/Runtime/Net/Iris/Private/Iris/ReplicationSystem/Filtering/ReplicationFiltering.cpp
+++ b/Engine/Source/Runtime/Net/Iris/Private/Iris/ReplicationSystem/Filtering/ReplicationFiltering.cpp
@@ -20,7 +20,6 @@
 #include "Iris/ReplicationSystem/ReplicationSystemInternal.h"
 #include "Iris/ReplicationSystem/Filtering/NetObjectGroups.h"
 #include "Iris/ReplicationSystem/Filtering/NetObjectFilterDefinitions.h"
-#include "Iris/ReplicationSystem/Filtering/NoOpNetObjectFilter.h"
 #include "Iris/ReplicationSystem/Filtering/ReplicationFilteringConfig.h"
 
 #include <inttypes.h>
@@ -233,7 +232,10 @@ void FReplicationFiltering::Deinit()
 {
 	for (FFilterInfo& FilterInfo : DynamicFilterInfos)
 	{
-		FilterInfo.Filter->Deinit();
+		if (FilterInfo.Filter)
+		{
+			FilterInfo.Filter->Deinit();
+		}
 	}
 
 	// Clear most buffers by setting size to 0
@@ -308,7 +310,10 @@ void FReplicationFiltering::OnMaxInternalNetRefIndexIncreased(FInternalNetRefInd
 		TArrayView<FNetObjectFilteringInfo> NewFilterInfoView = GetNetObjectFilteringInfos();
 		for (FFilterInfo& FilterInfo : DynamicFilterInfos)
 		{
-			FilterInfo.Filter->MaxInternalNetRefIndexIncreased(NewMaxInternalIndex, NewFilterInfoView);
+			if (FilterInfo.Filter)
+			{
+				FilterInfo.Filter->MaxInternalNetRefIndexIncreased(NewMaxInternalIndex, NewFilterInfoView);
+			}
 		}
 	}
 }
@@ -476,7 +481,7 @@ bool FReplicationFiltering::SetFilter(FInternalNetRefIndex ObjectIndex, FNetObje
 	const bool bWasUsingDynamicFilter = OldDynamicFilterIndex != InvalidDynamicFilterIndex;
 
 	// Validate the filter
-	if (bWantsToUseDynamicFilter && (NewDynamicFilterIndex >= (uint32)DynamicFilterInfos.Num()))
+	if (bWantsToUseDynamicFilter && !IsValidDynamicFilter(NewDynamicFilterIndex))
 	{
 		ensureMsgf(false, TEXT("Invalid dynamic filter 0x%08X. Filter is not being changed. Cause of ensure must be fixed!"), NewDynamicFilterIndex);
 		return false;
@@ -573,7 +578,7 @@ bool FReplicationFiltering::SetFilter(FInternalNetRefIndex ObjectIndex, FNetObje
 bool FReplicationFiltering::IsUsingSpatialFilter(FInternalNetRefIndex ObjectIndex) const
 {
 	const uint8 DynamicFilterIndex = ObjectIndexToDynamicFilterIndex[ObjectIndex];
-	if (DynamicFilterIndex == InvalidDynamicFilterIndex)
+	if (!IsValidDynamicFilter(DynamicFilterIndex))
 	{
 		return false;
 	}
@@ -631,7 +636,10 @@ void FReplicationFiltering::AddConnection(uint32 ConnectionId)
 
 	for (FFilterInfo& Info : DynamicFilterInfos)
 	{
-		Info.Filter->AddConnection(ConnectionId);
+		if (Info.Filter)
+		{
+			Info.Filter->AddConnection(ConnectionId);
+		}
 	}
 
 	// We defer needed processing to InitNewConnections().
@@ -652,7 +660,10 @@ void FReplicationFiltering::RemoveConnection(uint32 ConnectionId)
 
 	for (FFilterInfo& Info : DynamicFilterInfos)
 	{
-		Info.Filter->RemoveConnection(ConnectionId);
+		if (Info.Filter)
+		{
+			Info.Filter->RemoveConnection(ConnectionId);
+		}
 	}
 
 	// Reset SubObject filter for removed connection
@@ -1269,7 +1280,7 @@ void FReplicationFiltering::PreUpdateDynamicFiltering()
 	{
 		for (FFilterInfo& Info : DynamicFilterInfos)
 		{
-			if (Info.ObjectCount == 0U)
+			if (!Info.Filter || Info.ObjectCount == 0U)
 			{
 				continue;
 			}
@@ -1547,7 +1558,7 @@ void FReplicationFiltering::UpdateDynamicFiltering()
 		{
 			for (FFilterInfo& Info : DynamicFilterInfos)
 			{
-				if (Info.ObjectCount == 0U)
+				if (!Info.Filter || Info.ObjectCount == 0U)
 				{
 					continue;
 				}
@@ -1961,7 +1972,7 @@ void FReplicationFiltering::PostUpdateDynamicFiltering()
 		FNetObjectPostFilteringParams PostFilteringParams;
 		for (FFilterInfo& Info : DynamicFilterInfos)
 		{
-			if (Info.ObjectCount == 0U)
+			if (!Info.Filter || Info.ObjectCount == 0U)
 			{
 				continue;
 			}
@@ -2013,14 +2024,15 @@ void FReplicationFiltering::BatchNotifyFiltersOfDirtyObjects(FUpdateDirtyObjects
 
 	for (const FUpdateDirtyObjectsBatchHelper::FPerFilterInfo& PerFilterInfo : BatchHelper.PerFilterInfos)
 	{
-		if (PerFilterInfo.ObjectCount == 0)
+		const int32 FilterIndex = static_cast<int32>(&PerFilterInfo - BatchHelper.PerFilterInfos.GetData());
+
+		if (PerFilterInfo.ObjectCount == 0 || !IsValidDynamicFilter(FilterIndex))
 		{
 			continue;
 		}
 
 		UpdateParameters.ObjectIndices = PerFilterInfo.ObjectIndices;
 		UpdateParameters.ObjectCount = PerFilterInfo.ObjectCount;
-		const int32 FilterIndex = static_cast<int32>(&PerFilterInfo - BatchHelper.PerFilterInfos.GetData());
 		UNetObjectFilter* Filter = DynamicFilterInfos[FilterIndex].Filter.Get();
 		Filter->UpdateObjects(UpdateParameters);
 	}
@@ -3036,8 +3048,9 @@ void FReplicationFiltering::DestroyFilter(const FName Name)
 		// Destroy the filter but keep its index reserved to avoid renumbering subsequent filter handles and prevent reuse of stale handles with
 		// an unintended filter.
 		FilterInfo = FFilterInfo{
-			.Filter = MakeNoOpNetObjectFilter(),
+			.Filter = nullptr,
 			.Name = NAME_None,
+			.ObjectCount = 0,
 		};
 
 		break; // Filter names are unique; only one may match.
@@ -3049,6 +3062,13 @@ void FReplicationFiltering::InitObjectScopeHysteresis()
 	HysteresisState.Mode = Config->IsObjectScopeHysteresisEnabled() ? EHysteresisProcessingMode::Enabled : EHysteresisProcessingMode::Disabled;
 }
 
+bool FReplicationFiltering::IsValidDynamicFilter(uint32 FilterIndex) const
+{
+	return FilterIndex != InvalidDynamicFilterIndex
+		&& FilterIndex < static_cast<uint32>(DynamicFilterInfos.Num())
+		&& DynamicFilterInfos[FilterIndex].Filter != nullptr;
+}
+
 void FReplicationFiltering::RemoveFromDynamicFilter(uint32 ObjectIndex, uint32 FilterIndex)
 {
 	UE_LOG(LogIrisFiltering, Verbose, TEXT("RemoveFromDynamicFilter removing %s from Dynamic Filter %s"), 
@@ -3056,10 +3076,13 @@ void FReplicationFiltering::RemoveFromDynamicFilter(uint32 ObjectIndex, uint32 F
 
 	ObjectIndexToDynamicFilterIndex[ObjectIndex] = InvalidDynamicFilterIndex;
 	FNetObjectFilteringInfo& NetObjectFilteringInfo = NetObjectFilteringInfos[ObjectIndex];
-	FFilterInfo& FilterInfo = DynamicFilterInfos[FilterIndex];
-	--FilterInfo.ObjectCount;
-	FilterInfo.Filter->GetFilteredObjects().ClearBit(ObjectIndex);
-	FilterInfo.Filter->RemoveObject(ObjectIndex, NetObjectFilteringInfo);
+	if (IsValidDynamicFilter(FilterIndex))
+	{
+		FFilterInfo& FilterInfo = DynamicFilterInfos[FilterIndex];
+		--FilterInfo.ObjectCount;
+		FilterInfo.Filter->GetFilteredObjects().ClearBit(ObjectIndex);
+		FilterInfo.Filter->RemoveObject(ObjectIndex, NetObjectFilteringInfo);
+	}
 
 	DynamicFilterEnabledObjects.ClearBit(ObjectIndex);
 	ObjectsRequiringDynamicFilterUpdate.SetBit(ObjectIndex);
@@ -3146,7 +3169,7 @@ FString FReplicationFiltering::PrintFilterObjectInfo(FInternalNetRefIndex Object
 	}
 
 	const uint8 DynamicFilterIndex = ObjectIndexToDynamicFilterIndex[ObjectIndex];
-	if (DynamicFilterIndex == InvalidDynamicFilterIndex)
+	if (!IsValidDynamicFilter(DynamicFilterIndex))
 	{
 		return FString::Printf(TEXT("[NoDynamicFilter] %s"), PrintBuilder.ToString());
 	}
@@ -3177,7 +3200,10 @@ void FReplicationFiltering::BuildObjectsInFilterList(FNetBitArrayView OutObjects
 	{
 		if (FilterInfo.Name == FilterName)
 		{
-			OutObjectsInFilter.Copy(FilterInfo.Filter->GetFilteredObjects());
+			if (FilterInfo.Filter)
+			{
+				OutObjectsInFilter.Copy(FilterInfo.Filter->GetFilteredObjects());
+			}
 			return;
 		}
 	}
@@ -3282,20 +3308,6 @@ FNetObjectFilterInitParams FReplicationFiltering::GetFilterInitParams(UClass* Co
 	};
 }
 
-TStrongObjectPtr<UNetObjectFilter> FReplicationFiltering::MakeNoOpNetObjectFilter() const
-{
-	TStrongObjectPtr<UNetObjectFilter> Filter = TStrongObjectPtr<UNetObjectFilter>(NewObject<UNoOpNetObjectFilter>());
-
-	// Initialize the filter essentially to satisfy code analysis tools.
-	Filter->Init(GetFilterInitParams(/* ConfigClass */ nullptr));
-	ValidConnections.ForAllSetBits([&Filter = *Filter](const uint32 ConnectionId)
-		{
-			Filter.AddConnection(ConnectionId);
-		});
-
-	return Filter;
-}
-
 void FReplicationFiltering::FPerConnectionInfo::Deinit()
 {
 	ConnectionFilteredObjects.Empty();
diff --git a/Engine/Source/Runtime/Net/Iris/Private/Iris/ReplicationSystem/Filtering/ReplicationFiltering.h b/Engine/Source/Runtime/Net/Iris/Private/Iris/ReplicationSystem/Filtering/ReplicationFiltering.h
--- a/Engine/Source/Runtime/Net/Iris/Private/Iris/ReplicationSystem/Filtering/ReplicationFiltering.h
+++ b/Engine/Source/Runtime/Net/Iris/Private/Iris/ReplicationSystem/Filtering/ReplicationFiltering.h
@@ -215,6 +215,9 @@ private:
 
 	struct FFilterInfo
 	{
+		bool IsValid() const { return Filter != nullptr; };
+		explicit operator bool() const { return IsValid(); }
+
 		TStrongObjectPtr<UNetObjectFilter> Filter;
 		FName Name;
 		uint32 ObjectCount = 0;
@@ -302,6 +305,7 @@ private:
 
 	uint32 GetOwningConnectionIfDirty(uint32 ObjectIndex) const;
 
+	bool IsValidDynamicFilter(uint32 FilterIndex) const;
 	void RemoveFromDynamicFilter(uint32 ObjectIndex, uint32 FilterIndex);
 
 	void NotifyFiltersOfDirtyObjects();
@@ -431,6 +435,7 @@ private:
 	// Dynamic filters 
 	TArray<FNetObjectFilteringInfo> NetObjectFilteringInfos;
 	TArray<uint8> ObjectIndexToDynamicFilterIndex;
+	// Dynamic Filter infos. May contain invalid/null filters after dynamic filters are destroyed. Invalid/null filters must be ignored.
 	TArray<FFilterInfo> DynamicFilterInfos;
 
 	FNetBitArray DynamicFilterEnabledObjects;
