Iris
* Fix ReplicationSystem::GetReplicationBridge() so it returns a value even during initialization.

#rnx

#Synced-CL 51003732

diff --git a/Engine/Source/Runtime/Net/Iris/Private/Iris/ReplicationSystem/ReplicationSystem.cpp b/Engine/Source/Runtime/Net/Iris/Private/Iris/ReplicationSystem/ReplicationSystem.cpp
--- a/Engine/Source/Runtime/Net/Iris/Private/Iris/ReplicationSystem/ReplicationSystem.cpp
+++ b/Engine/Source/Runtime/Net/Iris/Private/Iris/ReplicationSystem/ReplicationSystem.cpp
@@ -1569,7 +1569,7 @@ UObject* UReplicationSystem::GetConnectionUserData(uint32 ConnectionId) const
 
 UObjectReplicationBridge* UReplicationSystem::GetReplicationBridge() const
 {
-	return Impl->ReplicationSystemInternal.GetReplicationBridge();
+	return ReplicationBridge;
 }
 
 bool UReplicationSystem::IsNetRefHandleAssigned(FNetRefHandle Handle) const
