Added support to UNetDriver for a multi-server network id without needing FRemoteObjectId.

#Synced-CL 51341769

diff --git a/Engine/Source/Runtime/CoreUObject/Public/UObject/RemoteObjectTypes.h b/Engine/Source/Runtime/CoreUObject/Public/UObject/RemoteObjectTypes.h
--- a/Engine/Source/Runtime/CoreUObject/Public/UObject/RemoteObjectTypes.h
+++ b/Engine/Source/Runtime/CoreUObject/Public/UObject/RemoteObjectTypes.h
@@ -369,7 +369,7 @@ public:
 	friend UE::RemoteObject::Private::FRemoteIdLocalizationHelper;
 	friend struct FUObjectItem;
 	friend class FNetGUIDCache;
-	friend class UNetDriver;
+	friend class URemoteObjectNetRefHandleGenerator;
 private:
 	
 	static bool IsAssetInternal(UObjectBase* InObject, EInternalObjectFlags InInitialFlags);
diff --git a/Engine/Source/Runtime/Engine/Private/NetDriver.cpp b/Engine/Source/Runtime/Engine/Private/NetDriver.cpp
--- a/Engine/Source/Runtime/Engine/Private/NetDriver.cpp
+++ b/Engine/Source/Runtime/Engine/Private/NetDriver.cpp
@@ -7669,32 +7669,6 @@ void UNetDriver::CreateReplicationSystem(bool bInitAsClient)
 			Params.InitialNetRefHandleIndex = GInitialNetworkId;
 		}
 
-#if UE_WITH_REMOTE_OBJECT_HANDLE
-		if (ProxyType == UE::Net::EProxyType::None)
-		{
-			Params.ObjectToNetRefHandleSerial.BindLambda([](const TObjectPtr<const UObject>& ObjectPtr)
-			{
-				const FRemoteObjectId RemoteObjectId = ObjectPtr.GetRemoteId();
-				const uint64 RemoteObjectIdValue = RemoteObjectId.GetIdNumber();
-
-				// Shift away the reserved bits in FRemoteObjectId so we're just left with the server id and serial number
-				// which will be embedded in the NetRefHandle serial.
-				const uint64 Serial = (RemoteObjectIdValue >> REMOTE_OBJECT_RESERVED_BIT_SIZE);
-
-				return Serial;
-			});
-
-			Params.NetRefHandleSerialToObject.BindLambda([this](const uint64 Serial)
-			{
-				// Shift left to made room for the reserved bits on FRemoteObjectId.
-				const uint64 RemoteObjectIdValue = (Serial << REMOTE_OBJECT_RESERVED_BIT_SIZE);
-				const FRemoteObjectId RemoteObjectId = FRemoteObjectId::CreateFromInt(RemoteObjectIdValue).GetLocalized();
-
-				return TObjectPtr<UObject>(FObjectPtr(RemoteObjectId));
-			});
-		}
-#endif
-
 		SetReplicationSystem(UE::Net::FReplicationSystemFactory::CreateReplicationSystem(Params));
 	}
 	else
diff --git a/Engine/Source/Runtime/Net/Iris/Private/Iris/ReplicationSystem/NetRefHandleGenerator.cpp b/Engine/Source/Runtime/Net/Iris/Private/Iris/ReplicationSystem/NetRefHandleGenerator.cpp
new file mode 100644
--- /dev/null
+++ b/Engine/Source/Runtime/Net/Iris/Private/Iris/ReplicationSystem/NetRefHandleGenerator.cpp
@@ -0,0 +1,103 @@
+// Copyright Epic Games, Inc. All Rights Reserved.
+
+#include "Iris/ReplicationSystem/NetRefHandleGenerator.h"
+#include "Iris/Core/BitTwiddling.h"
+#include "Iris/Core/IrisLog.h"
+#include "Iris/ReplicationSystem/NetRefHandle.h"
+
+void URemoteObjectNetRefHandleGenerator::Bind(UE::Net::FNetRefHandleGeneratorBindParams& Params)
+{
+#if UE_WITH_REMOTE_OBJECT_HANDLE
+	if (ensure(Params.ObjectToNetRefHandleSerial))
+	{
+		Params.ObjectToNetRefHandleSerial->BindLambda([](const TObjectPtr<const UObject>& ObjectPtr)
+		{
+			const FRemoteObjectId RemoteObjectId = ObjectPtr.GetRemoteId();
+			const uint64 RemoteObjectIdValue = RemoteObjectId.GetIdNumber();
+
+			// Shift away the reserved bits in FRemoteObjectId so we're just left with the server id and serial number
+			// which will be embedded in the NetRefHandle serial.
+			const uint64 Serial = (RemoteObjectIdValue >> REMOTE_OBJECT_RESERVED_BIT_SIZE);
+
+			UE_LOG(LogIris, Verbose, TEXT("URemoteObjectNetRefHandleGenerator: Mapping FNetRefHandle from FRemoteObjectId for object %s"), *ObjectPtr->GetName());
+
+			return Serial;
+		});
+	}
+
+	if (ensure(Params.NetRefHandleSerialToObject))
+	{
+		Params.NetRefHandleSerialToObject->BindLambda([](const uint64 Serial)
+		{
+			// Shift left to made room for the reserved bits on FRemoteObjectId.
+			const uint64 RemoteObjectIdValue = (Serial << REMOTE_OBJECT_RESERVED_BIT_SIZE);
+			const FRemoteObjectId RemoteObjectId = FRemoteObjectId::CreateFromInt(RemoteObjectIdValue).GetLocalized();
+
+			return TObjectPtr<UObject>(FObjectPtr(RemoteObjectId));
+		});
+	}
+#else
+	ensureMsgf(0, TEXT("URemoteObjectNetRefHandleGenerator: Requires UE_WITH_REMOTE_OBJECT_HANDLE to be defined."));
+	UE_LOG(LogIris, Error, TEXT("URemoteObjectNetRefHandleGenerator: Requires UE_WITH_REMOTE_OBJECT_HANDLE to be defined."));
+#endif
+}
+
+void UMultiServerNetRefHandleGenerator::Bind(UE::Net::FNetRefHandleGeneratorBindParams& Params)
+{
+	NextSerial = InitialSerial;
+
+	const uint64 LocalNumServerIdBits = NumServerIdBits;
+	const uint32 LocalServerId = ServerId;
+
+	if (LocalNumServerIdBits <= 0 || NumServerIdBits >= UE_NET_IRIS_NETREFHANDLE_SERIAL_SIZE)
+	{
+		UE_LOG(LogIris, Error, TEXT("UMultiServerNetRefHandleGenerator: An invalid value of NumServerIdBits was provided."));
+		return;
+	}
+
+	const uint64 NumSerialBits = UE_NET_IRIS_NETREFHANDLE_SERIAL_SIZE - NumServerIdBits;
+	if (NumSerialBits <= 0)
+	{
+		UE_LOG(LogIris, Error, TEXT("UMultiServerNetRefHandleGenerator: No space has been reserved for the serial number."))
+		return;
+	}
+
+	check(UE_NET_IRIS_NETREFHANDLE_SERIAL_SIZE <= 64);
+	check(LocalNumServerIdBits + NumSerialBits == UE_NET_IRIS_NETREFHANDLE_SERIAL_SIZE);
+
+	const uint64 LargestServerId = UE::Net::LargestBitValueForBits(LocalNumServerIdBits);
+	const uint64 LargestSerial = UE::Net::LargestBitValueForBits(NumSerialBits);
+
+	if (!ensureMsgf(ServerId < LargestServerId, TEXT("UMultiServerNetRefHandleGenerator: The server id %u is larger than the largest supported value of %llu"), ServerId, LargestServerId))
+	{
+		return;
+	}
+
+	UE_LOG(LogIris, Log, TEXT("UMultiServerNetRefHandleGenerator: ServerId=%u NumServerIdBits=%llu InitialSerial=%llu"), ServerId, LocalNumServerIdBits, NextSerial);
+
+	if (ensure(Params.ObjectToNetRefHandleSerial))
+	{
+		Params.ObjectToNetRefHandleSerial->BindLambda([this, LargestSerial, LocalNumServerIdBits, LocalServerId](const TObjectPtr<const UObject>& ObjectPtr)
+		{
+			const uint64 Serial = NextSerial;
+
+			if (!ensureMsgf(Serial < LargestSerial, TEXT("UMultiServerNetRefHandleGenerator: The serial %llu is larger than the largest supported value of %llu"), Serial, LargestSerial))
+			{
+				// This is equivelant to FNetRefHandle::GetInvalid().
+				return 0ULL;
+			}
+
+			NextSerial++;
+
+			const uint64 FinalNetRefHandleSerial = (Serial << LocalNumServerIdBits) | LocalServerId;
+
+			UE_LOG(LogIris, Verbose, TEXT("UMultiServerNetRefHandleGenerator: Assigning FNetRefHandle serial %llu (ServerId=%u Serial=%llu) for object %s"), 
+				FinalNetRefHandleSerial,
+				LocalServerId,
+				Serial,
+				*GetNameSafe(ObjectPtr));
+
+			return FinalNetRefHandleSerial;
+		});
+	}
+}
diff --git a/Engine/Source/Runtime/Net/Iris/Private/Iris/ReplicationSystem/ReplicationSystem.cpp b/Engine/Source/Runtime/Net/Iris/Private/Iris/ReplicationSystem/ReplicationSystem.cpp
--- a/Engine/Source/Runtime/Net/Iris/Private/Iris/ReplicationSystem/ReplicationSystem.cpp
+++ b/Engine/Source/Runtime/Net/Iris/Private/Iris/ReplicationSystem/ReplicationSystem.cpp
@@ -8,6 +8,8 @@
 #include "Iris/Core/IrisProfiler.h"
 #include "Iris/Core/NetObjectReference.h"
 
+#include "Iris/ReplicationSystem/ObjectReplicationBridgeConfig.h"
+
 #include "Net/Core/Connection/NetEnums.h"
 #include "Net/Core/Misc/NetConditionGroupManager.h"
 #include "Net/Core/NetToken/NetToken.h"
@@ -296,8 +298,20 @@ public:
 		{
 			FObjectReferenceCacheParams InitParams;
 			InitParams.bAuthority = ReplicationSystem->IsServer();
-			InitParams.ObjectToNetRefHandleSerial = Params.ObjectToNetRefHandleSerial;
-			InitParams.NetRefHandleSerialToObject = Params.NetRefHandleSerialToObject;
+
+			const bool bIsUsingSharedObjectReferenceCache = &ReplicationSystemInternal.GetObjectReferenceCache() == FReplicationSystemFactory::GetSharedObjectReferenceCache().Get();
+			if (!bIsUsingSharedObjectReferenceCache)
+			{
+				InitParams.ObjectToNetRefHandleSerial = Params.ObjectToNetRefHandleSerial;
+				InitParams.NetRefHandleSerialToObject = Params.NetRefHandleSerialToObject;
+			}
+			else
+			{
+				const bool bDelegatesAreBound = Params.ObjectToNetRefHandleSerial.IsBound() || Params.NetRefHandleSerialToObject.IsBound();
+				check(!bDelegatesAreBound);
+				UE_CLOG(bDelegatesAreBound, LogIris, Error, TEXT("Cannot bind FNetRefHandle mapping delegates when using a shared object reference cache."));
+			}
+
 			ReplicationSystemInternal.GetObjectReferenceCache().Init(InitParams);
 		}
 
@@ -923,8 +937,20 @@ void UReplicationSystem::Init(uint32 InId, const FReplicationSystemParams& Param
 
 	ReplicationBridge = Params.ReplicationBridge;
 
-	Impl = MakePimpl<UE::Net::Private::FReplicationSystemImpl>(this, Params);
-	Impl->Init(Params);
+	FReplicationSystemParams LocalParams = Params;
+
+	if (bIsServer)
+	{
+		InitializeNetRefHandleGenerator(LocalParams);
+	}
+	else
+	{
+		LocalParams.ObjectToNetRefHandleSerial.Unbind();
+		LocalParams.NetRefHandleSerialToObject.Unbind();
+	}
+
+	Impl = MakePimpl<UE::Net::Private::FReplicationSystemImpl>(this, LocalParams);
+	Impl->Init(LocalParams);
 
 	PostGarbageCollectHandle = FCoreUObjectDelegates::GetPostGarbageCollect().AddUObject(this, &UReplicationSystem::PostGarbageCollection);
 }
@@ -1249,6 +1275,45 @@ void UReplicationSystem::PostSeamlessTravelGarbageCollect()
 	Impl->ReplicationSystemInternal.PostSeamlessTravelGarbageCollect();
 }
 
+void UReplicationSystem::InitializeNetRefHandleGenerator(UReplicationSystem::FReplicationSystemParams& Params)
+{	
+	using namespace UE::Net;
+
+	if (Params.ProxyType != UE::Net::EProxyType::None)
+	{
+		return;
+	}
+
+	const UClass* GeneratorClass = UObjectReplicationBridgeConfig::GetConfig()->GetNetRefHandleGeneratorClass();
+	if (!GeneratorClass)
+	{
+		return;
+	}
+
+	if (!GeneratorClass->IsChildOf<UNetRefHandleGenerator>())
+	{
+		return;
+	}
+
+	NetRefHandleGenerator = NewObject<UNetRefHandleGenerator>(this, GeneratorClass);
+	if (!NetRefHandleGenerator)
+	{
+		UE_LOG(LogIris, Error, TEXT("UReplicationSystem::InitializeNetRefHandleGenerator: Unable to instantiate class %s"), *GeneratorClass->GetName());
+		return;
+	}
+
+	UE_LOG(LogIris, Log, TEXT("UReplicationSystem::InitializeNetRefHandleGenerator: Using generator %s"), *GeneratorClass->GetName());
+
+	UE::Net::FNetRefHandleGeneratorBindParams BindParams;
+	BindParams.ObjectToNetRefHandleSerial = &Params.ObjectToNetRefHandleSerial;
+	BindParams.NetRefHandleSerialToObject = &Params.NetRefHandleSerialToObject;
+
+	UE_CLOG(BindParams.ObjectToNetRefHandleSerial->IsBound(), LogIris, Warning, TEXT("UReplicationSystem::InitializeNetRefHandleGenerator: ObjectToNetRefHandleSerial is already set and will be overwritten."));
+	UE_CLOG(BindParams.NetRefHandleSerialToObject->IsBound(), LogIris, Warning, TEXT("UReplicationSystem::InitializeNetRefHandleGenerator: NetRefHandleSerialToObject is already set and will be overwritten."));
+
+	NetRefHandleGenerator->Bind(BindParams);
+}
+
 void UReplicationSystem::AddConnection(uint32 ConnectionId)
 {
 	Impl->AddConnection(ConnectionId);
diff --git a/Engine/Source/Runtime/Net/Iris/Public/Iris/ReplicationSystem/NetRefHandleGenerator.h b/Engine/Source/Runtime/Net/Iris/Public/Iris/ReplicationSystem/NetRefHandleGenerator.h
new file mode 100644
--- /dev/null
+++ b/Engine/Source/Runtime/Net/Iris/Public/Iris/ReplicationSystem/NetRefHandleGenerator.h
@@ -0,0 +1,80 @@
+// Copyright Epic Games, Inc. All Rights Reserved.
+
+#pragma once
+
+#include "HAL/Platform.h"
+#include "UObject/Object.h"
+#include "Net/Core/Connection/NetEnums.h"
+#include "Iris/ReplicationSystem/ReplicationSystemTypes.h"
+#include "NetRefHandleGenerator.generated.h"
+
+namespace UE::Net
+{
+	struct IRISCORE_API FNetRefHandleGeneratorBindParams
+	{
+		/**
+		 * If set, this delegate will override the delegate in FReplicationSystemParams::ObjectToNetRefHandleSerial.
+		 */
+		FObjectToNetRefHandleSerial* ObjectToNetRefHandleSerial = nullptr;
+
+		/**
+		 * If set, this delegate will override the delegate in FReplicationSystemParams::NetRefHandleSerialToObject.
+		 */
+		FNetRefHandleSerialToObject* NetRefHandleSerialToObject = nullptr;
+	};
+}
+
+UCLASS(Abstract, Transient)
+class IRISCORE_API UNetRefHandleGenerator : public UObject
+{
+	GENERATED_BODY()
+
+public:
+	virtual ~UNetRefHandleGenerator() = default;
+
+	/** 
+	 * This function is used to bind bespoke functions to the delegates in Params.
+	 */
+	virtual void Bind(UE::Net::FNetRefHandleGeneratorBindParams& Params) PURE_VIRTUAL(UNetRefHandleGenerator::Bind, );
+};
+
+/** 
+ * Generate a FNetRefHandle serial using an object's corresponding FRemoteObjectId.
+ *
+ * This generator requires UE_WITH_REMOTE_OBJECT_HANDLE to be defined.
+ */
+UCLASS()
+class IRISCORE_API URemoteObjectNetRefHandleGenerator : public UNetRefHandleGenerator
+{
+	GENERATED_BODY()
+
+public:
+	virtual void Bind(UE::Net::FNetRefHandleGeneratorBindParams& Params) override;
+};
+
+/** 
+ * Generate a FNetRefHandle serial for use in a multi-server environment that isn't using UE_WITH_REMOTE_OBJECT_HANDLE.
+ *
+ * The FNetRefHandle serial will reserve bits for a server id, and the remaining bits of the serial
+ * will be used for an auto-incrementing serial number.
+ */ 
+UCLASS(Config=Engine, DefaultConfig)
+class IRISCORE_API UMultiServerNetRefHandleGenerator : public UNetRefHandleGenerator
+{
+	GENERATED_BODY()
+
+public:
+	virtual void Bind(UE::Net::FNetRefHandleGeneratorBindParams& Params) override;
+
+private:
+	uint64 NextSerial = 1;
+
+	UPROPERTY(Config)
+	uint64 NumServerIdBits = 4;
+
+	UPROPERTY(Config)
+	uint32 ServerId = 0;
+
+	UPROPERTY(Config)
+	uint32 InitialSerial = 1;
+};
diff --git a/Engine/Source/Runtime/Net/Iris/Public/Iris/ReplicationSystem/ObjectReplicationBridgeConfig.h b/Engine/Source/Runtime/Net/Iris/Public/Iris/ReplicationSystem/ObjectReplicationBridgeConfig.h
--- a/Engine/Source/Runtime/Net/Iris/Public/Iris/ReplicationSystem/ObjectReplicationBridgeConfig.h
+++ b/Engine/Source/Runtime/Net/Iris/Public/Iris/ReplicationSystem/ObjectReplicationBridgeConfig.h
@@ -4,8 +4,10 @@
 
 #include "Containers/ContainersFwd.h"
 #include "UObject/ObjectMacros.h"
+#include "Templates/SubclassOf.h"
 
 #include "Iris/ReplicationSystem/ObjectReferenceTypes.h"
+#include "Iris/ReplicationSystem/NetRefHandleGenerator.h"
 
 #include "ObjectReplicationBridgeConfig.generated.h"
 
@@ -153,6 +155,8 @@ public:
 	IRISCORE_API TConstArrayView<FObjectReplicationBridgeTypeStatsConfig> GetTypeStatsConfigs() const;
 	IRISCORE_API TConstArrayView<FObjectReplicationBridgeAsyncLoadingClassConfig> GetAsyncLoadingClassConfigs() const;
 
+	UClass* GetNetRefHandleGeneratorClass() const { return NetRefHandleGeneratorClass; }
+
 	FName GetDefaultSpatialFilterName() const;
 	FName GetRequiredNetDriverChannelClassName() const;
 
@@ -224,6 +228,14 @@ private:
 	/** Classes that must be instantiated otherwise we need to disconnect from the session entirely.  */
 	UPROPERTY(Config)
 	TArray<FName> CriticalActorClasses;
+	
+	/** 
+	 * An optional configuration property defines how a NetRefHandle is generated.
+	 *
+	 * If not set, the NetRefHandle will be generated as an incremental value starting with 1.
+	 */
+	UPROPERTY(Config)
+	TSubclassOf<UNetRefHandleGenerator> NetRefHandleGeneratorClass;
 };
 
 inline FName UObjectReplicationBridgeConfig::GetDefaultSpatialFilterName() const
diff --git a/Engine/Source/Runtime/Net/Iris/Public/Iris/ReplicationSystem/ReplicationSystem.h b/Engine/Source/Runtime/Net/Iris/Public/Iris/ReplicationSystem/ReplicationSystem.h
--- a/Engine/Source/Runtime/Net/Iris/Public/Iris/ReplicationSystem/ReplicationSystem.h
+++ b/Engine/Source/Runtime/Net/Iris/Public/Iris/ReplicationSystem/ReplicationSystem.h
@@ -10,6 +10,7 @@
 #include "Iris/ReplicationSystem/NetRefHandle.h"
 #include "Iris/ReplicationSystem/NetObjectGroupHandle.h"
 #include "Iris/ReplicationSystem/ReplicationSystemTypes.h"
+#include "Iris/ReplicationSystem/NetRefHandleGenerator.h"
 #include "Net/Core/Connection/NetEnums.h"
 #include "Net/Core/Connection/ConnectionHandle.h"
 #include "Net/Core/NetHandle/NetHandle.h"
@@ -31,6 +32,7 @@ class UNetObjectFilter;
 class UNetObjectPrioritizer;
 class UObjectReplicationBridge;
 class UReplicationSystem;
+class UNetRefHandleGenerator;
 
 struct FNetObjectFilterDefinition;
 
@@ -875,6 +877,8 @@ private:
 
 	void SetPIEInstanceID(int32 InPIEInstanceID) { PIEInstanceID = InPIEInstanceID; }
 
+	void InitializeNetRefHandleGenerator(UReplicationSystem::FReplicationSystemParams& Params);
+
 private:
 
 	void PostGarbageCollection();
@@ -899,6 +903,9 @@ private:
 	uint32 bReplicateTransactionally : 1;
 
 	UE::Net::EProxyType ProxyType = UE::Net::EProxyType::None;
+
+	UPROPERTY()
+	TObjectPtr<UNetRefHandleGenerator> NetRefHandleGenerator;
 };
 
 namespace UE::Net
