Iris - Allow any module to use FIrisLogOnceTracker

#okforgithuib public
#rnx

#Synced-CL 51332317

diff --git a/Engine/Source/Runtime/Net/Iris/Public/Iris/Core/IrisLogUtils.h b/Engine/Source/Runtime/Net/Iris/Public/Iris/Core/IrisLogUtils.h
--- a/Engine/Source/Runtime/Net/Iris/Public/Iris/Core/IrisLogUtils.h
+++ b/Engine/Source/Runtime/Net/Iris/Public/Iris/Core/IrisLogUtils.h
@@ -12,16 +12,16 @@ namespace UE::Net
 
 /**
  * Utility class to determine whether some message coupled to a hash hasn't been logged yet. 
- * Example usage would be for exmaple if you only want to log something once per class. 
+ * Example usage would be for example if you only want to log something once per class. 
  */
 class FIrisLogOnceTracker
 {
 public:
 	/** Returns true if the Hash has not been encountered in a call to ShouldLog before. */
-	bool ShouldLog(uint32 Hash) const;
+	IRISCORE_API bool ShouldLog(uint32 Hash) const;
 
 	/** Returns true if the hash of the Name has not been encountered in a call to ShouldLog before. */
-	bool ShouldLog(const FName& Name) const;
+	IRISCORE_API bool ShouldLog(const FName& Name) const;
 
 private:
 	mutable TSet<uint32> LoggedHashes;
