A fix to Iris that avoids starting a new transaction if it's already being called from within a transaction.

#Synced-CL 51341769

diff --git a/Engine/Source/Runtime/Net/Iris/Private/Iris/ReplicationSystem/NetObjectFactory.cpp b/Engine/Source/Runtime/Net/Iris/Private/Iris/ReplicationSystem/NetObjectFactory.cpp
--- a/Engine/Source/Runtime/Net/Iris/Private/Iris/ReplicationSystem/NetObjectFactory.cpp
+++ b/Engine/Source/Runtime/Net/Iris/Private/Iris/ReplicationSystem/NetObjectFactory.cpp
@@ -64,7 +64,7 @@ TUniquePtr<UE::Net::FNetObjectCreationHeader> UNetObjectFactory::CreateHeader(UE
 	// Ask the derived class to allocate and fill the header
 	TUniquePtr<FNetObjectCreationHeader> Header;
 #if UE_WITH_REMOTE_OBJECT_HANDLE && UE_AUTORTFM
-	if (UE::Net::Private::FBridgeSerialization::IsReplicatingTransactionally(Bridge))
+	if (UE::Net::Private::FBridgeSerialization::IsReplicatingTransactionally(Bridge) && !AutoRTFM::IsTransactional())
 	{
 		UE::RemoteExecutor::ExecuteTransactional(UE::Net::Private::CreateAndFillHeaderTransactionName, [this, &Header, Handle]()
 			{
