Iris
* Deleted legacy location state indexes data

#rnx
#Synced-CL 50927993

diff --git a/Engine/Source/Runtime/Net/Iris/Private/Iris/ReplicationSystem/Filtering/NetObjectGridFilter.cpp b/Engine/Source/Runtime/Net/Iris/Private/Iris/ReplicationSystem/Filtering/NetObjectGridFilter.cpp
--- a/Engine/Source/Runtime/Net/Iris/Private/Iris/ReplicationSystem/Filtering/NetObjectGridFilter.cpp
+++ b/Engine/Source/Runtime/Net/Iris/Private/Iris/ReplicationSystem/Filtering/NetObjectGridFilter.cpp
@@ -77,7 +77,7 @@ uint16 UNetObjectGridFilter::GetFrameCountBeforeCulling(FName ProfileName) const
 
 bool UNetObjectGridFilter::AddObject(uint32 ObjectIndex, FNetObjectFilterAddObjectParams& Params)
 {
-	// We support either a world location in the state, tagged with RepTag_WorldLocation, or via the WorldLocations instance.
+	// We only support objects that registered a world location
 	const bool bHasLocation = BuildObjectInfo(ObjectIndex, Params);
 
 	if (!bHasLocation)
@@ -602,8 +602,6 @@ void UNetObjectGridWorldLocFilter::PreFilter(FNetObjectPreFilteringParams& Param
 
 void UNetObjectGridWorldLocFilter::UpdateObjectInfo(UNetObjectGridFilter::FPerObjectInfo& PerObjectInfo, const UNetObjectGridFilter::FObjectLocationInfo& ObjectLocationInfo, const UE::Net::FReplicationInstanceProtocol* InstanceProtocol)
 {
-	check(ObjectLocationInfo.IsUsingWorldLocations());
-
 	const TOptional<UE::Net::FWorldLocations::FWorldInfo> Info = WorldLocations->GetWorldInfo(PerObjectInfo.ObjectIndex);
 	
 	if (Info)
@@ -617,11 +615,6 @@ bool UNetObjectGridWorldLocFilter::BuildObjectInfo(uint32 ObjectIndex, FNetObjec
 {
 	checkf(WorldLocations->HasInfoForObject(ObjectIndex), TEXT("UNetObjectGridWorldLocFilter only supports objects with valid world location data."));
 
-	// Craft tag info that will let us know we need to retrieve the location from WorldLocations
-	FObjectLocationInfo& ObjectLocationInfo = static_cast<FObjectLocationInfo&>(Params.OutInfo);
-	ObjectLocationInfo.SetLocationStateOffset(InvalidStateOffset);
-	ObjectLocationInfo.SetLocationStateIndex(InvalidStateIndex);
-
 	return true;
 }
 
diff --git a/Engine/Source/Runtime/Net/Iris/Public/Iris/ReplicationSystem/Filtering/NetObjectGridFilter.h b/Engine/Source/Runtime/Net/Iris/Public/Iris/ReplicationSystem/Filtering/NetObjectGridFilter.h
--- a/Engine/Source/Runtime/Net/Iris/Public/Iris/ReplicationSystem/Filtering/NetObjectGridFilter.h
+++ b/Engine/Source/Runtime/Net/Iris/Public/Iris/ReplicationSystem/Filtering/NetObjectGridFilter.h
@@ -100,23 +100,13 @@ protected:
 protected:
 	struct FObjectLocationInfo : public FNetObjectFilteringInfo
 	{
-		bool IsUsingWorldLocations() const  { return GetLocationStateIndex() == InvalidStateIndex; }
-		bool IsUsingLocationInState() const { return GetLocationStateIndex() != InvalidStateIndex; }
-
 		/**
 		* Data mapping:
-		* uint16 Data[0] = LocationState offset
-		* uint16 Data[1] = LocationState index
+		* uint16 Data[0] = Not used
+		* uint16 Data[1] = Not used
 		* uint16 Data[2] = FPerObjectInfo index (low bytes)
 		* uint16 Data[3] = FPerObjectInfo index (high bytes)
 		*/
-
-		void SetLocationStateOffset(uint16 Offset)	{ Data[0] = Offset; }
-		uint16 GetLocationStateOffset() const		{ return Data[0]; }
-
-		void SetLocationStateIndex(uint16 Index) { Data[1] = Index; }
-		uint16 GetLocationStateIndex() const	 { return Data[1]; }
-
 		void SetInfoIndex(uint32 Index)	{ Data[2] = Index & 65535U; Data[3] = Index >> 16U; }
 		uint32 GetInfoIndex() const		{ return (uint32(Data[3]) << 16U) | uint32(Data[2]); }
 	};
@@ -130,8 +120,6 @@ protected:
 	enum : unsigned
 	{
 		ObjectInfosChunkSize = 64 * 1024,
-		InvalidStateIndex = 65535U,
-		InvalidStateOffset = 65535U,
 	};
 
 	struct FCellBox
