Iris
- Fix to invalidate stale references to default subobjects of dynamic replicated objects.
- Can be disabled with cvar "net.iris.VerifyOuterForStaleDynamicReferences 0"

diff --git a/Engine/Source/Runtime/Net/Iris/Private/Iris/ReplicationSystem/ObjectReferenceCache.cpp b/Engine/Source/Runtime/Net/Iris/Private/Iris/ReplicationSystem/ObjectReferenceCache.cpp
--- a/Engine/Source/Runtime/Net/Iris/Private/Iris/ReplicationSystem/ObjectReferenceCache.cpp
+++ b/Engine/Source/Runtime/Net/Iris/Private/Iris/ReplicationSystem/ObjectReferenceCache.cpp
@@ -82,6 +82,9 @@ FAutoConsoleVariableRef CVarIrisAsyncLoadingFailAllLoads(TEXT("net.Iris.AsyncLoa
 static bool bIrisEnableRemoveReferenceFallbackPath = true;
 FAutoConsoleVariableRef CVarIrisEnableRemoveReferenceFallbackPath(TEXT("net.iris.ObjectReferenceCache.EnableRemoveReferenceFallbackPath"), bIrisEnableRemoveReferenceFallbackPath, TEXT("Whether to enable the slow fallback path or not when ObjectReferenceCache::RemoveReference is passed a null pointer object."), ECVF_Default);
 
+static bool bIrisVerifyOuterForStaleDynamicReferences = true;
+FAutoConsoleVariableRef CVarIrisVerifyOuterForStaleDynamicReferences(TEXT("net.iris.VerifyOuterForStaleDynamicReferences"), bIrisVerifyOuterForStaleDynamicReferences, TEXT("If set, we will verify outer for when creationg stably named references of dynamic objects. Default is true."), ECVF_Default);
+
 /**
  * Don't allow infinite recursion of InternalLoadObject - an attacker could
  * send malicious packets that cause a stack overflow on the server.
@@ -343,39 +346,76 @@ bool FObjectReferenceCache::CreateObjectReferenceInternal(const UObject* Object,
 			FCachedNetObjectReference& CachedObject = *CachedObjectPtr;
 			if (CachedObject.Object.Get() == Object)
 			{
-				UE_LOG_REFERENCECACHE(VeryVerbose, TEXT("ObjectReferenceCache::CreateObjectReferenceHandle Found existing %s for ObjectPath %s Object: %s (0x%p), OuterNetRefHandle: %s"),
-				*CachedObject.NetRefHandle.ToString(), ToCStr(Object->GetPathName()), *GetNameSafe(Object), Object, *CachedObject.OuterNetRefHandle.ToString());
-				OutReference = MakeNetObjectReference(CachedObject);
-				return true;
-			}
-			else
-			{
-				UE_LOG_REFERENCECACHE(Verbose, TEXT("ObjectReferenceCache::CreateObjectReferenceHandle Removed %s from ObjectToNetReferenceHandle due to stale cache. Cached Object (0x%p).  New object %s (0x%p)"),
-					*CachedObject.NetRefHandle.ToString(), CachedObject.ObjectKey, *GetNameSafe(Object), Object);
-
-				ObjectToNetReferenceHandle.Remove(CachedObject.ObjectKey);
-				ObjectToNetReferenceHandle.Remove(Object);
-				RefHandlePtr = nullptr;
-
-				if (RefHandle.IsStatic())
+				bool bOuterChainValid = true;
+				// For relative dynamic references we also need to verify the cached outers
+				if (bIrisVerifyOuterForStaleDynamicReferences && (RefHandle.IsDynamic() && CachedObject.RelativePath.IsValid()))
 				{
-					UE_LOG_REFERENCECACHE(Verbose, TEXT("ObjectReferenceCache::CreateObjectReferenceHandle clearing stale cache for static handle %s. Cached Object (0x%p).  New object %s (0x%p)"),
-						*CachedObject.NetRefHandle.ToString(), CachedObject.ObjectKey, *GetNameSafe(Object), Object);
-
-					// Note we only cleanse the object reference
-					// we still keep the cachedObject data around in order to be able to serialize static destruction infos
-					CachedObject.ObjectKey = nullptr;
-					CachedObject.Object = nullptr;
+					const UObject* CurrentOuter = Object->GetOuter();
+					FNetRefHandle CurrentCachedOuterHandle = CachedObject.OuterNetRefHandle;
 
+					while (CurrentOuter != nullptr || CurrentCachedOuterHandle.IsValid())
+					{
+						const FCachedNetObjectReference* CachedOuterPtr = CurrentCachedOuterHandle.IsValid() ? ReferenceHandleToCachedReference.Find(CurrentCachedOuterHandle) : nullptr;
+						if (CachedOuterPtr)
+						{
+							if (CachedOuterPtr->Object.Get() != CurrentOuter)
+							{
+								bOuterChainValid = false;
+								break;
+							}
+
+							// No need to go deeper.
+							if (!(CurrentCachedOuterHandle.IsDynamic() && CachedOuterPtr->RelativePath.IsValid()))
+							{
+								break;
+							}
+							
+							CurrentOuter = CurrentOuter ? CurrentOuter->GetOuter() : nullptr;
+							CurrentCachedOuterHandle = CachedOuterPtr->OuterNetRefHandle;
+						}
+						else
+						{
+							bOuterChainValid = false;
+							break;
+						}
+					}
 				}
-				else
+
+				if (bOuterChainValid)
 				{
-					UE_LOG_REFERENCECACHE(Verbose, TEXT("ObjectReferenceCache::CreateObjectReferenceHandle removing cache for handle %s. Cached Object (0x%p). New object %s (0x%p)"),
-						*CachedObject.NetRefHandle.ToString(), CachedObject.ObjectKey, *GetNameSafe(Object), Object);
+					UE_LOG_REFERENCECACHE(VeryVerbose, TEXT("ObjectReferenceCache::CreateObjectReferenceHandle Found existing %s for ObjectPath %s Object: %s (0x%p), OuterNetRefHandle: %s"),
+					*CachedObject.NetRefHandle.ToString(), ToCStr(Object->GetPathName()), *GetNameSafe(Object), Object, *CachedObject.OuterNetRefHandle.ToString());
+					OutReference = MakeNetObjectReference(CachedObject);
 
-					ReferenceHandleToCachedReference.Remove(RefHandle);
+					return true;
 				}
 			}
+
+			UE_LOG_REFERENCECACHE(Verbose, TEXT("ObjectReferenceCache::CreateObjectReferenceHandle Removed %s from ObjectToNetReferenceHandle due to stale cache. Cached Object (0x%p).  New object %s (0x%p)"),
+				*CachedObject.NetRefHandle.ToString(), CachedObject.ObjectKey, *GetNameSafe(Object), Object);
+
+			ObjectToNetReferenceHandle.Remove(CachedObject.ObjectKey);
+			ObjectToNetReferenceHandle.Remove(Object);
+			RefHandlePtr = nullptr;
+
+			if (RefHandle.IsStatic())
+			{
+				UE_LOG_REFERENCECACHE(Verbose, TEXT("ObjectReferenceCache::CreateObjectReferenceHandle clearing stale cache for static handle %s. Cached Object (0x%p).  New object %s (0x%p)"),
+					*CachedObject.NetRefHandle.ToString(), CachedObject.ObjectKey, *GetNameSafe(Object), Object);
+
+				// Note we only cleanse the object reference
+				// we still keep the cachedObject data around in order to be able to serialize static destruction infos
+				CachedObject.ObjectKey = nullptr;
+				CachedObject.Object = nullptr;
+
+			}
+			else
+			{
+				UE_LOG_REFERENCECACHE(Verbose, TEXT("ObjectReferenceCache::CreateObjectReferenceHandle removing cache for handle %s. Cached Object (0x%p). New object %s (0x%p)"),
+					*CachedObject.NetRefHandle.ToString(), CachedObject.ObjectKey, *GetNameSafe(Object), Object);
+
+				ReferenceHandleToCachedReference.Remove(RefHandle);
+			}
 		}
 		else
 		{
@@ -1420,7 +1460,7 @@ void FObjectReferenceCache::WriteFullReferenceInternal(FNetSerializationContext&
 		UE_CLOG(RefHandle.IsValid(), LogIrisReferences, Log, TEXT("ObjectReferenceCache::WriteFullReference Trying to write Stale handle, %s"), *RefHandle.ToString());
 
 		WriteNetRefHandle(Context, FNetRefHandle::GetInvalid());
-		return;			
+		return;
 	}
 
 	// Write handle
