Iris
- Fixed asan error

#rnx
#Synced-CL 51584953

diff --git a/Engine/Source/Runtime/Net/Iris/Private/Iris/DataStream/DataStreamManager.cpp b/Engine/Source/Runtime/Net/Iris/Private/Iris/DataStream/DataStreamManager.cpp
--- a/Engine/Source/Runtime/Net/Iris/Private/Iris/DataStream/DataStreamManager.cpp
+++ b/Engine/Source/Runtime/Net/Iris/Private/Iris/DataStream/DataStreamManager.cpp
@@ -365,13 +365,6 @@ UDataStreamManager::EWriteResult UDataStreamManager::FImpl::WriteData(UE::Net::F
 		return EWriteResult::NoData;
 	}
 
-	// Init export record
-	NetExports.InitExportRecordForPacket();
-
-	// Setup export context for this packet
-	FNetExportContext::FBatchExports CurrentPacketBatchExports;
-	FNetExports::FExportScope ExportScope = NetExports.MakeExportScope(Context, CurrentPacketBatchExports);
-
 	FRecord TempRecord;
 	TempRecord.DataStreamRecords.SetNumZeroed(StreamCount);
 	FDataStreamRecord const** TempStreamRecords = TempRecord.DataStreamRecords.GetData();
@@ -439,92 +432,101 @@ UDataStreamManager::EWriteResult UDataStreamManager::FImpl::WriteData(UE::Net::F
 		return EWriteResult::NoData;
 	}
 
-	uint32 DataStreamMask = 0;
-	for (uint32 StreamIt = 0U, StreamEndIt = StreamCount, CurrentStreamMask = 1U; StreamIt != StreamEndIt; ++StreamIt, CurrentStreamMask += CurrentStreamMask)
 	{
-		UDataStream* Stream = StreamData[StreamIt];
-		const EDataStreamSendStatus SendStatus = SendStatusData[StreamIt];
-		if (SendStatus == EDataStreamSendStatus::Pause)
-		{
-			continue;
-		}
+		// Init export record
+		NetExports.InitExportRecordForPacket();
+
+		// Setup export context for this packet
+		FNetExportContext::FBatchExports CurrentPacketBatchExports;
+		FNetExports::FExportScope ExportScope = NetExports.MakeExportScope(Context, CurrentPacketBatchExports);
 
-		// We only write stream data of the stream is considered is open
-		const EDataStreamState State = StreamState[StreamIt];
-		if (!(State == EDataStreamState::Open || State == EDataStreamState::PendingClose))
+		uint32 DataStreamMask = 0;
+		for (uint32 StreamIt = 0U, StreamEndIt = StreamCount, CurrentStreamMask = 1U; StreamIt != StreamEndIt; ++StreamIt, CurrentStreamMask += CurrentStreamMask)
 		{
-			continue;
-		}
+			UDataStream* Stream = StreamData[StreamIt];
+			const EDataStreamSendStatus SendStatus = SendStatusData[StreamIt];
+			if (SendStatus == EDataStreamSendStatus::Pause)
+			{
+				continue;
+			}
 
-		FNetBitStreamWriter SubBitStream = ManagerStream.CreateSubstream();
-		FNetSerializationContext SubContext = Context.MakeSubContext(&SubBitStream);
+			// We only write stream data of the stream is considered is open
+			const EDataStreamState State = StreamState[StreamIt];
+			if (!(State == EDataStreamState::Open || State == EDataStreamState::PendingClose))
+			{
+				continue;
+			}
 
-		const FDataStreamRecord* SubRecord = nullptr;
-		const EWriteResult WriteResult = Stream->WriteData(SubContext, SubRecord);
-		if (WriteResult == EWriteResult::NoData || SubContext.HasError())
-		{
-			checkf(SubRecord == nullptr, TEXT("DataStream '%s' provided a record despite errors or returning NoData."), ToCStr(Stream->GetFName().GetPlainNameString()));
-			ManagerStream.DiscardSubstream(SubBitStream);
+			FNetBitStreamWriter SubBitStream = ManagerStream.CreateSubstream();
+			FNetSerializationContext SubContext = Context.MakeSubContext(&SubBitStream);
 
-			if (SubContext.HasError())
+			const FDataStreamRecord* SubRecord = nullptr;
+			const EWriteResult WriteResult = Stream->WriteData(SubContext, SubRecord);
+			if (WriteResult == EWriteResult::NoData || SubContext.HasError())
 			{
-				Context.SetError(SubContext.GetError(), false);
-				break;
+				checkf(SubRecord == nullptr, TEXT("DataStream '%s' provided a record despite errors or returning NoData."), ToCStr(Stream->GetFName().GetPlainNameString()));
+				ManagerStream.DiscardSubstream(SubBitStream);
+
+				if (SubContext.HasError())
+				{
+					Context.SetError(SubContext.GetError(), false);
+					break;
+				}
+				else
+				{
+					continue;
+				}
+			}
+
+			// Only update DataStreamMask if data was written. 
+			if (SubBitStream.GetPosBits() > 0U)
+			{
+				DataStreamMask |= CurrentStreamMask;		
+				TempStreamRecords[StreamIt] = SubRecord;
 			}
 			else
 			{
-				continue;
+				ensureMsgf(SubRecord == nullptr, TEXT("DataStream '%s' provided a record despite not writing any data."), ToCStr(Stream->GetFName().GetPlainNameString()));
 			}
+
+			ManagerStream.CommitSubstream(SubBitStream);
+
+			// Set CombinedWriteResult to HasMoreData if any of the result variables is HasMoreData, otherwise take the WriteResult, which will be 'Ok'.
+			CombinedWriteResult = (CombinedWriteResult == EWriteResult::HasMoreData || WriteResult == EWriteResult::HasMoreData) ? EWriteResult::HasMoreData : WriteResult;
 		}
 
-		// Only update DataStreamMask if data was written. 
-		if (SubBitStream.GetPosBits() > 0U)
+		if (!(DataStreamMask || bHasStreamsWithDirtyState))
 		{
-			DataStreamMask |= CurrentStreamMask;		
-			TempStreamRecords[StreamIt] = SubRecord;
+			Context.GetBitStreamWriter()->DiscardSubstream(ManagerStream);
+			// Technically we could also return EWriteResult::HasMoreData
+			CombinedWriteResult = EWriteResult::NoData;
 		}
 		else
 		{
-			ensureMsgf(SubRecord == nullptr, TEXT("DataStream '%s' provided a record despite not writing any data."), ToCStr(Stream->GetFName().GetPlainNameString()));
-		}
-
-		ManagerStream.CommitSubstream(SubBitStream);
-
-		// Set CombinedWriteResult to HasMoreData if any of the result variables is HasMoreData, otherwise take the WriteResult, which will be 'Ok'.
-		CombinedWriteResult = (CombinedWriteResult == EWriteResult::HasMoreData || WriteResult == EWriteResult::HasMoreData) ? EWriteResult::HasMoreData : WriteResult;
-	}
-
-	if (!(DataStreamMask || bHasStreamsWithDirtyState))
-	{
-		Context.GetBitStreamWriter()->DiscardSubstream(ManagerStream);
-		// Technically we could also return EWriteResult::HasMoreData
-		CombinedWriteResult = EWriteResult::NoData;
-	}
-	else
-	{
-		// Fixup manager header
-		{
-			const uint32 CurrentBitPos = ManagerStream.GetPosBits();
-			ManagerStream.Seek(0);
-			ManagerStream.WriteBits(StreamCount - 1U, StreamCountBitCount);
-			ManagerStream.WriteBits(DataStreamMask, StreamCount);
-			ManagerStream.Seek(CurrentBitPos);
-			Context.GetBitStreamWriter()->CommitSubstream(ManagerStream);
-		}
+			// Fixup manager header
+			{
+				const uint32 CurrentBitPos = ManagerStream.GetPosBits();
+				ManagerStream.Seek(0);
+				ManagerStream.WriteBits(StreamCount - 1U, StreamCountBitCount);
+				ManagerStream.WriteBits(DataStreamMask, StreamCount);
+				ManagerStream.Seek(CurrentBitPos);
+				Context.GetBitStreamWriter()->CommitSubstream(ManagerStream);
+			}
 
-		// Fixup and store record
-		TempRecord.DataStreamMask = DataStreamMask;
-		TempRecord.DataStreamStateMask = DirtyStreamsMask;
-		DirtyStreamsMask = 0U;
+			// Fixup and store record
+			TempRecord.DataStreamMask = DataStreamMask;
+			TempRecord.DataStreamStateMask = DirtyStreamsMask;
+			DirtyStreamsMask = 0U;
 
-		FRecord*& Record = Records.Enqueue_GetRef();
-		*Record = MoveTemp(TempRecord);
+			FRecord*& Record = Records.Enqueue_GetRef();
+			*Record = MoveTemp(TempRecord);
 
-		OutRecord = Record;
+			OutRecord = Record;
 
-		// Push exports and update export record
-		NetExports.CommitExportsToRecord(ExportScope);
-		NetExports.PushExportRecordForPacket();
+			// Push exports and update export record
+			NetExports.CommitExportsToRecord(ExportScope);
+			NetExports.PushExportRecordForPacket();
+		}
 	}
 
 	return CombinedWriteResult;
