Iris
* Added NetRefInternalIndexManager that will be used to expose internal Iris netbitarrays in NetObjectFilter public functions
* Renamed FNetRefHandleManager::InvalidInternalIndex to UE::Net::InvalidInternalNetRefIndex
* Moved FInternalNetRefIndex typedef from UE::Net::Private to UE::Net
* Changed the NetObjectFIlter virtual functions that referenced object indexes from uint32 to UE::Net::FInternalNetRefIndex

#rnx
#Synced-CL 51624330

diff --git a/Engine/Plugins/Runtime/ReplicationSystemTestPlugin/Source/Private/Tests/ReplicationSystem/Filtering/MockNetObjectFilter.cpp b/Engine/Plugins/Runtime/ReplicationSystemTestPlugin/Source/Private/Tests/ReplicationSystem/Filtering/MockNetObjectFilter.cpp
--- a/Engine/Plugins/Runtime/ReplicationSystemTestPlugin/Source/Private/Tests/ReplicationSystem/Filtering/MockNetObjectFilter.cpp
+++ b/Engine/Plugins/Runtime/ReplicationSystemTestPlugin/Source/Private/Tests/ReplicationSystem/Filtering/MockNetObjectFilter.cpp
@@ -29,7 +29,7 @@ void UMockNetObjectFilter::OnDeinit()
 {
 }
 
-void UMockNetObjectFilter::OnMaxInternalNetRefIndexIncreased(uint32 NewMaxInternalIndex)
+void UMockNetObjectFilter::OnMaxInternalNetRefIndexIncreased(UE::Net::FInternalNetRefIndex NewMaxInternalIndex)
 {
 	AddedObjectIndices.SetNumBits(NewMaxInternalIndex);
 }
@@ -51,7 +51,7 @@ void UMockNetObjectFilter::RemoveConnection(uint32 ConnectionId)
 	AddedConnectionIndices.ClearBit(ConnectionId);
 }
 
-bool UMockNetObjectFilter::AddObject(uint32 ObjectIndex, FNetObjectFilterAddObjectParams& Params)
+bool UMockNetObjectFilter::AddObject(UE::Net::FInternalNetRefIndex ObjectIndex, FNetObjectFilterAddObjectParams& Params)
 {
 	++CallStatus.CallCounts.AddObject;
 	if (CallSetup.AddObject.bReturnValue)
@@ -79,7 +79,7 @@ bool UMockNetObjectFilter::AddObject(uint32 ObjectIndex, FNetObjectFilterAddObje
 	return CallSetup.AddObject.bReturnValue;
 }
 
-void UMockNetObjectFilter::RemoveObject(uint32 ObjectIndex, const FNetObjectFilteringInfo& Info)
+void UMockNetObjectFilter::RemoveObject(UE::Net::FInternalNetRefIndex ObjectIndex, const FNetObjectFilteringInfo& Info)
 {
 	++CallStatus.CallCounts.RemoveObject;
 
@@ -106,7 +106,7 @@ void UMockNetObjectFilter::UpdateObjects(FNetObjectFilterUpdateParams& Params)
 	++CallStatus.CallCounts.UpdateObjects;
 
 	bool bIsProperCall = Params.ObjectCount <= AddedCount;
-	for (uint32 ObjectIndex : MakeArrayView(Params.ObjectIndices, Params.ObjectCount))
+	for (UE::Net::FInternalNetRefIndex ObjectIndex : MakeArrayView(Params.ObjectIndices, Params.ObjectCount))
 	{
 		bIsProperCall = bIsProperCall && AddedObjectIndices.GetBit(ObjectIndex);
 	}
@@ -159,7 +159,7 @@ void UMockNetObjectFilterWithCondition::OnInit(const FNetObjectFilterInitParams&
 	ReplicationSystem = Params.ReplicationSystem.Get();
 }
 
-bool UMockNetObjectFilterWithCondition::AddObject(uint32 ObjectIndex, FNetObjectFilterAddObjectParams& Params)
+bool UMockNetObjectFilterWithCondition::AddObject(UE::Net::FInternalNetRefIndex ObjectIndex, FNetObjectFilterAddObjectParams& Params)
 {
 	bool bResult = Super::AddObject(ObjectIndex, Params);
 
@@ -176,7 +176,7 @@ void UMockNetObjectFilterWithCondition::Filter(FNetObjectFilteringParams& Params
 
 	bool bIsProperCall = true;
 
-	FilteredObjects.ForAllSetBits([&](uint32 ObjectIndex)
+	FilteredObjects.ForAllSetBits([&](UE::Net::FInternalNetRefIndex ObjectIndex)
 	{
 		UTestFilteringObject* FilterObject = CastChecked<UTestFilteringObject>(ReplicationSystem->GetReplicationSystemInternal()->GetNetRefHandleManager().GetReplicatedObjectInstance(ObjectIndex));
 		
diff --git a/Engine/Plugins/Runtime/ReplicationSystemTestPlugin/Source/Private/Tests/ReplicationSystem/Filtering/MockNetObjectFilter.h b/Engine/Plugins/Runtime/ReplicationSystemTestPlugin/Source/Private/Tests/ReplicationSystem/Filtering/MockNetObjectFilter.h
--- a/Engine/Plugins/Runtime/ReplicationSystemTestPlugin/Source/Private/Tests/ReplicationSystem/Filtering/MockNetObjectFilter.h
+++ b/Engine/Plugins/Runtime/ReplicationSystemTestPlugin/Source/Private/Tests/ReplicationSystem/Filtering/MockNetObjectFilter.h
@@ -7,7 +7,7 @@
 #include "Net/Core/NetBitArray.h"
 #include "MockNetObjectFilter.generated.h"
 
-namespace UE::Net::Private
+namespace UE::Net
 {
 	typedef uint32 FInternalNetRefIndex;
 }
@@ -68,11 +68,11 @@ public:
 protected:
 	virtual void OnInit(const FNetObjectFilterInitParams&) override;
 	virtual void OnDeinit() override;
-	virtual void OnMaxInternalNetRefIndexIncreased(uint32 NewMaxInternalIndex) override;
+	virtual void OnMaxInternalNetRefIndexIncreased(UE::Net::FInternalNetRefIndex NewMaxInternalIndex) override;
 	virtual void AddConnection(uint32 ConnectionId) override;
 	virtual void RemoveConnection(uint32 ConnectionId) override;
-	virtual bool AddObject(uint32 ObjectIndex, FNetObjectFilterAddObjectParams&) override;
-	virtual void RemoveObject(uint32 ObjectIndex, const FNetObjectFilteringInfo&) override;
+	virtual bool AddObject(UE::Net::FInternalNetRefIndex ObjectIndex, FNetObjectFilterAddObjectParams&) override;
+	virtual void RemoveObject(UE::Net::FInternalNetRefIndex ObjectIndex, const FNetObjectFilteringInfo&) override;
 	virtual void UpdateObjects(FNetObjectFilterUpdateParams&) override;
 	virtual void PreFilter(FNetObjectPreFilteringParams&) override;
 	virtual void Filter(FNetObjectFilteringParams&) override;
@@ -101,7 +101,7 @@ class UMockNetObjectFilterWithCondition : public UMockNetObjectFilter
 protected:
 
 	virtual void OnInit(const FNetObjectFilterInitParams&) override;
-	virtual bool AddObject(uint32 ObjectIndex, FNetObjectFilterAddObjectParams&) override;
+	virtual bool AddObject(UE::Net::FInternalNetRefIndex ObjectIndex, FNetObjectFilterAddObjectParams&) override;
 	virtual void Filter(FNetObjectFilteringParams&) override;
 
 private:
diff --git a/Engine/Plugins/Runtime/ReplicationSystemTestPlugin/Source/Private/Tests/ReplicationSystem/Filtering/TestGroups.cpp b/Engine/Plugins/Runtime/ReplicationSystemTestPlugin/Source/Private/Tests/ReplicationSystem/Filtering/TestGroups.cpp
--- a/Engine/Plugins/Runtime/ReplicationSystemTestPlugin/Source/Private/Tests/ReplicationSystem/Filtering/TestGroups.cpp
+++ b/Engine/Plugins/Runtime/ReplicationSystemTestPlugin/Source/Private/Tests/ReplicationSystem/Filtering/TestGroups.cpp
@@ -60,7 +60,7 @@ protected:
 	void VerifyObjectFilterStatus(FNetRefHandle Handle, uint32 ConnectionId, bool bExpectedReplicationAllowed)
 	{
 		const uint32 InternalIndex = Server->GetReplicationSystem()->GetReplicationSystemInternal()->GetNetRefHandleManager().GetInternalIndex(Handle);
-		UE_NET_ASSERT_NE((uint32)FNetRefHandleManager::InvalidInternalIndex, InternalIndex);
+		UE_NET_ASSERT_NE((uint32)InvalidInternalNetRefIndex, InternalIndex);
 
 		// Verify expected filter status
 		const FNetBitArrayView ObjectGroupFilter = ServerFiltering->GetGroupFilteredOutObjects(ConnectionId);
@@ -459,7 +459,7 @@ UE_NET_TEST_FIXTURE(FTestGroupsFixture, GroupFilterTestObjectListViaTrait)
 	const FNetRefHandle ServerHandle = ServerObject0->NetRefHandle;
 	
 	const uint32 ServerInternalIndex = Server->GetReplicationSystem()->GetReplicationSystemInternal()->GetNetRefHandleManager().GetInternalIndex(ServerHandle);
-	UE_NET_ASSERT_NE((uint32)FNetRefHandleManager::InvalidInternalIndex, ServerInternalIndex);
+	UE_NET_ASSERT_NE((uint32)InvalidInternalNetRefIndex, ServerInternalIndex);
 
 	// Create groups
 	FNetObjectGroupHandle GroupHandle0 = Server->ReplicationSystem->CreateGroup(NAME_None);
@@ -529,7 +529,7 @@ UE_NET_TEST_FIXTURE(FTestGroupsFixture, GroupFilterTestObjectListViaMembership)
 	const FNetRefHandle ServerHandle = ServerObject0->NetRefHandle;
 
 	const uint32 ServerInternalIndex = Server->GetReplicationSystem()->GetReplicationSystemInternal()->GetNetRefHandleManager().GetInternalIndex(ServerHandle);
-	UE_NET_ASSERT_NE((uint32)FNetRefHandleManager::InvalidInternalIndex, ServerInternalIndex);
+	UE_NET_ASSERT_NE((uint32)InvalidInternalNetRefIndex, ServerInternalIndex);
 
 	// Create groups
 	FNetObjectGroupHandle GroupHandle0 = Server->ReplicationSystem->CreateGroup(NAME_None);
diff --git a/Engine/Plugins/Runtime/ReplicationSystemTestPlugin/Source/Private/Tests/ReplicationSystem/Prioritization/MockNetObjectPrioritizer.cpp b/Engine/Plugins/Runtime/ReplicationSystemTestPlugin/Source/Private/Tests/ReplicationSystem/Prioritization/MockNetObjectPrioritizer.cpp
--- a/Engine/Plugins/Runtime/ReplicationSystemTestPlugin/Source/Private/Tests/ReplicationSystem/Prioritization/MockNetObjectPrioritizer.cpp
+++ b/Engine/Plugins/Runtime/ReplicationSystemTestPlugin/Source/Private/Tests/ReplicationSystem/Prioritization/MockNetObjectPrioritizer.cpp
@@ -18,12 +18,12 @@ void UMockNetObjectPrioritizer::Init(FNetObjectPrioritizerInitParams& Params)
 	AddedIndices.Init(Params.CurrentMaxInternalIndex);
 }
 
-void UMockNetObjectPrioritizer::OnMaxInternalNetRefIndexIncreased(uint32 NewMaxInternalIndex)
+void UMockNetObjectPrioritizer::OnMaxInternalNetRefIndexIncreased(UE::Net::FInternalNetRefIndex NewMaxInternalIndex)
 {
 	AddedIndices.SetNumBits(NewMaxInternalIndex);
 }
 
-bool UMockNetObjectPrioritizer::AddObject(uint32 ObjectIndex, FNetObjectPrioritizerAddObjectParams& Params)
+bool UMockNetObjectPrioritizer::AddObject(UE::Net::FInternalNetRefIndex ObjectIndex, FNetObjectPrioritizerAddObjectParams& Params)
 {
 	++CallStatus.CallCounts.AddObject;
 	if (CallSetup.AddObject.ReturnValue == true)
@@ -60,7 +60,7 @@ bool UMockNetObjectPrioritizer::AddObject(uint32 ObjectIndex, FNetObjectPrioriti
 	return CallSetup.AddObject.ReturnValue;
 }
 
-void UMockNetObjectPrioritizer::RemoveObject(uint32 ObjectIndex, const FNetObjectPrioritizationInfo& Info)
+void UMockNetObjectPrioritizer::RemoveObject(UE::Net::FInternalNetRefIndex ObjectIndex, const FNetObjectPrioritizationInfo& Info)
 {
 	++CallStatus.CallCounts.RemoveObject;
 
@@ -115,8 +115,8 @@ void UMockNetObjectPrioritizer::Prioritize(FNetObjectPrioritizationParams& Param
 	bIsProperCall &= Params.ObjectCount == AddedCount;
 
 	// Validate that we only need to prioritize objects that were added, that they're in order and unique.
-	uint32 LastObjectIndex = 0;
-	for (uint32 ObjectIndex : MakeArrayView(Params.ObjectIndices, Params.ObjectCount))
+	UE::Net::FInternalNetRefIndex LastObjectIndex = 0;
+	for (UE::Net::FInternalNetRefIndex ObjectIndex : MakeArrayView(Params.ObjectIndices, Params.ObjectCount))
 	{
 		bIsProperCall &= (AddedIndices.IsBitSet(ObjectIndex));
 		bIsProperCall &= (ObjectIndex > LastObjectIndex);
@@ -137,7 +137,7 @@ void UMockNetObjectPrioritizer::Prioritize(FNetObjectPrioritizationParams& Param
 	CallStatus.SuccessfulCallCounts.Prioritize += bIsProperCall;
 }
 
-float UMockNetObjectPrioritizer::GetPriority(UE::Net::Private::FInternalNetRefIndex ObjectIndex) const
+float UMockNetObjectPrioritizer::GetPriority(UE::Net::FInternalNetRefIndex ObjectIndex) const
 {
 	check(AddedIndices.IsBitSet(ObjectIndex));
 	if (const float* Priority = ObjectToPriority.Find(ObjectIndex))
diff --git a/Engine/Plugins/Runtime/ReplicationSystemTestPlugin/Source/Private/Tests/ReplicationSystem/Prioritization/MockNetObjectPrioritizer.h b/Engine/Plugins/Runtime/ReplicationSystemTestPlugin/Source/Private/Tests/ReplicationSystem/Prioritization/MockNetObjectPrioritizer.h
--- a/Engine/Plugins/Runtime/ReplicationSystemTestPlugin/Source/Private/Tests/ReplicationSystem/Prioritization/MockNetObjectPrioritizer.h
+++ b/Engine/Plugins/Runtime/ReplicationSystemTestPlugin/Source/Private/Tests/ReplicationSystem/Prioritization/MockNetObjectPrioritizer.h
@@ -7,7 +7,7 @@
 #include "Iris/ReplicationSystem/RepTag.h"
 #include "MockNetObjectPrioritizer.generated.h"
 
-namespace UE::Net::Private
+namespace UE::Net
 {
 	typedef uint32 FInternalNetRefIndex;
 }
@@ -56,14 +56,14 @@ public:
 	inline const FFunctionCallStatus& GetFunctionCallStatus() const { return CallStatus; }
 	inline void ResetFunctionCallStatus() { CallStatus = FFunctionCallStatus({}); }
 
-	float GetPriority(UE::Net::Private::FInternalNetRefIndex ObjectIndex) const;
+	float GetPriority(UE::Net::FInternalNetRefIndex ObjectIndex) const;
 
 protected:
 	virtual void Init(FNetObjectPrioritizerInitParams& Params) override;
 	virtual void Deinit() override {}
-	virtual void OnMaxInternalNetRefIndexIncreased(uint32 NewMaxInternalIndex) override;
-	virtual bool AddObject(uint32 ObjectIndex, FNetObjectPrioritizerAddObjectParams& Params) override;
-	virtual void RemoveObject(uint32 ObjectIndex, const FNetObjectPrioritizationInfo& Info) override;
+	virtual void OnMaxInternalNetRefIndexIncreased(UE::Net::FInternalNetRefIndex NewMaxInternalIndex) override;
+	virtual bool AddObject(UE::Net::FInternalNetRefIndex ObjectIndex, FNetObjectPrioritizerAddObjectParams& Params) override;
+	virtual void RemoveObject(UE::Net::FInternalNetRefIndex ObjectIndex, const FNetObjectPrioritizationInfo& Info) override;
 	virtual void UpdateObjects(FNetObjectPrioritizerUpdateParams&) override;
 	virtual void Prioritize(FNetObjectPrioritizationParams&) override;
 
@@ -74,8 +74,8 @@ private:
 	FFunctionCallSetup CallSetup;
 	UE::Net::FNetBitArray AddedIndices;
 	SIZE_T AddedCount;
-	TMap<uint32, UPTRINT> ObjectToPriorityOffset;
-	TMap<uint32, float> ObjectToPriority;
+	TMap<UE::Net::FInternalNetRefIndex, UPTRINT> ObjectToPriorityOffset;
+	TMap<UE::Net::FInternalNetRefIndex, float> ObjectToPriority;
 
 	static constexpr float DefaultPriority = 1.0f;
 };
diff --git a/Engine/Plugins/Runtime/ReplicationSystemTestPlugin/Source/Private/Tests/ReplicationSystem/Prioritization/TestNetObjectPrioritizer.cpp b/Engine/Plugins/Runtime/ReplicationSystemTestPlugin/Source/Private/Tests/ReplicationSystem/Prioritization/TestNetObjectPrioritizer.cpp
--- a/Engine/Plugins/Runtime/ReplicationSystemTestPlugin/Source/Private/Tests/ReplicationSystem/Prioritization/TestNetObjectPrioritizer.cpp
+++ b/Engine/Plugins/Runtime/ReplicationSystemTestPlugin/Source/Private/Tests/ReplicationSystem/Prioritization/TestNetObjectPrioritizer.cpp
@@ -455,7 +455,7 @@ UE_NET_TEST_FIXTURE(FTestNetObjectPrioritizerAPIFixture, NativeIrisObjectGetsPri
 	ReplicationSystem->NetUpdate(FakeDeltaTime);
 	ReplicationSystem->PostSendUpdate();
 
-	UE::Net::Private::FInternalNetRefIndex InternalIndex = NetRefHandleManager->GetInternalIndex(NetRefHandle);
+	UE::Net::FInternalNetRefIndex InternalIndex = NetRefHandleManager->GetInternalIndex(NetRefHandle);
 	const float Priority = MockNetObjectPrioritizer->GetPriority(InternalIndex);
 	UE_NET_ASSERT_EQ(Priority, ObjectPriority);
 }
@@ -482,7 +482,7 @@ UE_NET_TEST_FIXTURE(FTestNetObjectPrioritizerAPIFixture, NativeIrisObjectGetsUpd
 	ReplicationSystem->NetUpdate(FakeDeltaTime);
 	ReplicationSystem->PostSendUpdate();
 
-	UE::Net::Private::FInternalNetRefIndex InternalIndex = NetRefHandleManager->GetInternalIndex(NetRefHandle);
+	UE::Net::FInternalNetRefIndex InternalIndex = NetRefHandleManager->GetInternalIndex(NetRefHandle);
 	const float Priority = MockNetObjectPrioritizer->GetPriority(InternalIndex);
 	UE_NET_ASSERT_EQ(Priority, UpdatedPriority);
 }
@@ -506,7 +506,7 @@ UE_NET_TEST_FIXTURE(FTestNetObjectPrioritizerAPIFixture, ObjectGetsPriorityFromS
 	ReplicationSystem->NetUpdate(FakeDeltaTime);
 	ReplicationSystem->PostSendUpdate();
 
-	UE::Net::Private::FInternalNetRefIndex InternalIndex = NetRefHandleManager->GetInternalIndex(NetRefHandle);
+	UE::Net::FInternalNetRefIndex InternalIndex = NetRefHandleManager->GetInternalIndex(NetRefHandle);
 	const float Priority = MockNetObjectPrioritizer->GetPriority(InternalIndex);
 	UE_NET_ASSERT_EQ(Priority, ObjectPriority);
 }
@@ -533,7 +533,7 @@ UE_NET_TEST_FIXTURE(FTestNetObjectPrioritizerAPIFixture, ObjectGetsUpdatedPriori
 	ReplicationSystem->NetUpdate(FakeDeltaTime);
 	ReplicationSystem->PostSendUpdate();
 
-	UE::Net::Private::FInternalNetRefIndex InternalIndex = NetRefHandleManager->GetInternalIndex(NetRefHandle);
+	UE::Net::FInternalNetRefIndex InternalIndex = NetRefHandleManager->GetInternalIndex(NetRefHandle);
 	const float Priority = MockNetObjectPrioritizer->GetPriority(InternalIndex);
 	UE_NET_ASSERT_EQ(Priority, UpdatedPriority);
 }
diff --git a/Engine/Plugins/Runtime/ReplicationSystemTestPlugin/Source/Private/Tests/ReplicationSystem/Prioritization/TestNetObjectPrioritizerFixture.cpp b/Engine/Plugins/Runtime/ReplicationSystemTestPlugin/Source/Private/Tests/ReplicationSystem/Prioritization/TestNetObjectPrioritizerFixture.cpp
--- a/Engine/Plugins/Runtime/ReplicationSystemTestPlugin/Source/Private/Tests/ReplicationSystem/Prioritization/TestNetObjectPrioritizerFixture.cpp
+++ b/Engine/Plugins/Runtime/ReplicationSystemTestPlugin/Source/Private/Tests/ReplicationSystem/Prioritization/TestNetObjectPrioritizerFixture.cpp
@@ -115,8 +115,8 @@ FTestNetObjectPrioritizerFixture::FPrioritizationResult FTestNetObjectPrioritize
 			for (UTestPrioritizationWithWorldLocationObject* Object : WorldLocObjects)
 			{
 				float Priority = 0.0f;
-				Private::FInternalNetRefIndex NetRefIndex = NetRefHandleManager.GetInternalIndex(Object->NetRefHandle);
-				if (NetRefIndex != Private::FNetRefHandleManager::InvalidInternalIndex)
+				FInternalNetRefIndex NetRefIndex = NetRefHandleManager.GetInternalIndex(Object->NetRefHandle);
+				if (NetRefIndex != InvalidInternalNetRefIndex)
 				{
 					Priority = Priorities[NetRefIndex];
 				}
@@ -167,8 +167,8 @@ FTestNetObjectPrioritizerFixture::FPrioritizationResult FTestNetObjectPrioritize
 			for (const FNetRefHandle& NetRefHandle : NetRefHandles)
 			{
 				float Priority = 0.0f;
-				Private::FInternalNetRefIndex NetRefIndex = NetRefHandleManager.GetInternalIndex(NetRefHandle);
-				if (NetRefIndex != Private::FNetRefHandleManager::InvalidInternalIndex)
+				FInternalNetRefIndex NetRefIndex = NetRefHandleManager.GetInternalIndex(NetRefHandle);
+				if (NetRefIndex != InvalidInternalNetRefIndex)
 				{
 					Priority = Priorities[NetRefIndex];
 				}
@@ -262,8 +262,8 @@ FNetObjectPrioritizerImage FTestNetObjectPrioritizerFixture::Visualize(FNetObjec
 					for (UTestPrioritizationWithWorldLocationObject* Object : WorldLocObjects)
 					{
 						float Priority = 0.0f;
-						Private::FInternalNetRefIndex NetRefIndex = NetRefHandleManager.GetInternalIndex(Object->NetRefHandle);
-						if (NetRefIndex != Private::FNetRefHandleManager::InvalidInternalIndex)
+						FInternalNetRefIndex NetRefIndex = NetRefHandleManager.GetInternalIndex(Object->NetRefHandle);
+						if (NetRefIndex != InvalidInternalNetRefIndex)
 						{
 							Priority = Priorities[NetRefIndex];
 						}
diff --git a/Engine/Plugins/Runtime/ReplicationSystemTestPlugin/Source/Private/Tests/ReplicationSystem/TestDirtyNetObjectTracker.cpp b/Engine/Plugins/Runtime/ReplicationSystemTestPlugin/Source/Private/Tests/ReplicationSystem/TestDirtyNetObjectTracker.cpp
--- a/Engine/Plugins/Runtime/ReplicationSystemTestPlugin/Source/Private/Tests/ReplicationSystem/TestDirtyNetObjectTracker.cpp
+++ b/Engine/Plugins/Runtime/ReplicationSystemTestPlugin/Source/Private/Tests/ReplicationSystem/TestDirtyNetObjectTracker.cpp
@@ -55,7 +55,7 @@ UE_NET_TEST_FIXTURE(FDirtyNetObjectTrackerTestFixture, CannotMarkInvalidObjectAs
 {
 	LOG_SCOPE_VERBOSITY_OVERRIDE(LogIrisDirtyTracker, ELogVerbosity::Error);
 
-	MarkNetObjectStateDirty(ReplicationSystemId, FNetRefHandleManager::InvalidInternalIndex);
+	MarkNetObjectStateDirty(ReplicationSystemId, InvalidInternalNetRefIndex);
 	MarkNetObjectStateDirty(ReplicationSystemId, NetObjectIndexRangeEnd + 1);
 
 	FDirtyObjectsAccessor DirtyObjectsAccessor(*DirtyNetObjectTracker);
diff --git a/Engine/Source/Runtime/Net/Iris/Private/Iris/Core/IrisDebugging.cpp b/Engine/Source/Runtime/Net/Iris/Private/Iris/Core/IrisDebugging.cpp
--- a/Engine/Source/Runtime/Net/Iris/Private/Iris/Core/IrisDebugging.cpp
+++ b/Engine/Source/Runtime/Net/Iris/Private/Iris/Core/IrisDebugging.cpp
@@ -57,7 +57,7 @@ namespace IrisDebugHelperInternal
 			}
 		}));
 
-	static UE::Net::Private::FInternalNetRefIndex GIrisDebugInternalIndex = UE::Net::Private::FNetRefHandleManager::InvalidInternalIndex;
+	static UE::Net::FInternalNetRefIndex GIrisDebugInternalIndex = InvalidInternalNetRefIndex;
 	static FAutoConsoleCommand NetIrisDebugNetInternalIndex(
 		TEXT("Net.Iris.DebugNetInternalIndex"),
 		TEXT("Specify an internal index that we will break on (or none to turn off)."),
@@ -120,7 +120,7 @@ bool BreakOnRPCName(FName RPCName)
 }
 
 
-bool BreakOnInternalNetRefIndex(UE::Net::Private::FInternalNetRefIndex InternalIndex)
+bool BreakOnInternalNetRefIndex(UE::Net::FInternalNetRefIndex InternalIndex)
 {
 	if (IrisDebugHelperInternal::GIrisDebugInternalIndex != 0 && IrisDebugHelperInternal::GIrisDebugInternalIndex == InternalIndex)
 	{
@@ -131,7 +131,7 @@ bool BreakOnInternalNetRefIndex(UE::Net::Private::FInternalNetRefIndex InternalI
 	return false;
 }
 
-UE::Net::Private::FInternalNetRefIndex GetDebugInternalNetRefIndex()
+UE::Net::FInternalNetRefIndex GetDebugInternalNetRefIndex()
 {
 	return IrisDebugHelperInternal::GIrisDebugInternalIndex;
 }
@@ -170,7 +170,7 @@ void SetIrisDebugRPCName(const ANSICHAR* NameBuffer)
 	}
 }
 
-void SetIrisDebugInternalNetRefIndex(UE::Net::Private::FInternalNetRefIndex InternalIndex)
+void SetIrisDebugInternalNetRefIndex(UE::Net::FInternalNetRefIndex InternalIndex)
 {
 	IrisDebugHelperInternal::GIrisDebugInternalIndex = InternalIndex;
 }
@@ -182,7 +182,7 @@ void SetIrisDebugInternalNetRefIndexViaNetHandle(FNetRefHandle RefHandle)
 	if (!RefHandle.IsValid())
 	{
 		IrisDebugHelperInternal::GIrisDebugNetRefHandle = RefHandle;
-		IrisDebugHelperInternal::GIrisDebugInternalIndex = FNetRefHandleManager::InvalidInternalIndex;
+		IrisDebugHelperInternal::GIrisDebugInternalIndex = InvalidInternalNetRefIndex;
 		return;
 	}
 
@@ -212,7 +212,7 @@ void SetIrisDebugInternalNetRefIndexViaObject(UObject* Instance)
 	if (!Instance)
 	{
 		IrisDebugHelperInternal::GIrisDebugNetRefHandle = FNetRefHandle();
-		IrisDebugHelperInternal::GIrisDebugInternalIndex = FNetRefHandleManager::InvalidInternalIndex;
+		IrisDebugHelperInternal::GIrisDebugInternalIndex = InvalidInternalNetRefIndex;
 		return;
 	}
 
@@ -437,7 +437,7 @@ FNetReplicatedObjectDebugInfo DebugNetRefHandleById(uint64 NetRefHandleId, uint3
 			const FNetRefHandleManager& NetRefHandleManager = ReplicationSystem->GetReplicationSystemInternal()->GetNetRefHandleManager();
 			FNetRefHandle IncompleteHandle = FNetRefHandleManager::MakeNetRefHandle(NetRefHandleId, ReplicationSystemId);
 			const uint32 InternalNetRefIndex = Info.InternalNetRefIndex = NetRefHandleManager.GetInternalIndex(IncompleteHandle);
-			if (InternalNetRefIndex != FNetRefHandleManager::InvalidInternalIndex)
+			if (InternalNetRefIndex != InvalidInternalNetRefIndex)
 			{
 				const FNetRefHandleManager::FReplicatedObjectData& ObjectData = NetRefHandleManager.GetReplicatedObjectDataNoCheck(InternalNetRefIndex);
 				if (ObjectData.RefHandle == IncompleteHandle)
@@ -469,7 +469,7 @@ FNetReplicatedObjectDebugInfo DebugInternalNetRefIndex(uint32 InternalIndex, uin
 		if (ReplicationSystem)
 		{
 			const FNetRefHandleManager& NetRefHandleManager = ReplicationSystem->GetReplicationSystemInternal()->GetNetRefHandleManager();
-			if (InternalIndex != FNetRefHandleManager::InvalidInternalIndex)
+			if (InternalIndex != InvalidInternalNetRefIndex)
 			{
 				FNetRefHandle ObjectHandle = NetRefHandleManager.GetNetRefHandleFromInternalIndex(InternalIndex);
 				const FNetRefHandleManager::FReplicatedObjectData& ObjectData = NetRefHandleManager.GetReplicatedObjectDataNoCheck(InternalIndex);
diff --git a/Engine/Source/Runtime/Net/Iris/Private/Iris/ReplicationState/ReplicationStateStorage.cpp b/Engine/Source/Runtime/Net/Iris/Private/Iris/ReplicationState/ReplicationStateStorage.cpp
--- a/Engine/Source/Runtime/Net/Iris/Private/Iris/ReplicationState/ReplicationStateStorage.cpp
+++ b/Engine/Source/Runtime/Net/Iris/Private/Iris/ReplicationState/ReplicationStateStorage.cpp
@@ -47,7 +47,7 @@ void FReplicationStateStorage::Deinit()
 #endif
 }
 
-void FReplicationStateStorage::OnMaxInternalNetRefIndexIncreased(UE::Net::Private::FInternalNetRefIndex NewMaxInternalIndex)
+void FReplicationStateStorage::OnMaxInternalNetRefIndexIncreased(UE::Net::FInternalNetRefIndex NewMaxInternalIndex)
 {
 	ObjectIndexToObjectInfoIndex.SetNumZeroed(NewMaxInternalIndex);
 }
diff --git a/Engine/Source/Runtime/Net/Iris/Private/Iris/ReplicationState/ReplicationStateStorage.h b/Engine/Source/Runtime/Net/Iris/Private/Iris/ReplicationState/ReplicationStateStorage.h
--- a/Engine/Source/Runtime/Net/Iris/Private/Iris/ReplicationState/ReplicationStateStorage.h
+++ b/Engine/Source/Runtime/Net/Iris/Private/Iris/ReplicationState/ReplicationStateStorage.h
@@ -15,9 +15,11 @@
 namespace UE::Net
 {
 	struct FReplicationProtocol;
+	
+	typedef uint32 FInternalNetRefIndex;
+
 	namespace Private
 	{
-		typedef uint32 FInternalNetRefIndex;
 		class FNetRefHandleManager;
 	};
 }
@@ -52,7 +54,7 @@ public:
 	public:
 		bool IsValid() const { return ReservedStorage != nullptr; }
 
-		// The uninitalized memory that was reserved.
+		// The uninitialized memory that was reserved.
 		uint8* ReservedStorage = nullptr;
 		/*
 		 * The state presumed to be cloned if the reservation is committed. N.B.
@@ -100,7 +102,7 @@ public:
 	void CommitBaselineReservation(uint32 ObjectIndex, uint8* Storage, EReplicationStateType Base);
 
 	/** Called when the maximum InternalNetRefIndex increased and we need to realloc our lists */
-	void OnMaxInternalNetRefIndexIncreased(UE::Net::Private::FInternalNetRefIndex NewMaxInternalIndex);
+	void OnMaxInternalNetRefIndexIncreased(UE::Net::FInternalNetRefIndex NewMaxInternalIndex);
 
 private:
 	using ObjectInfoIndexType = uint16;
diff --git a/Engine/Source/Runtime/Net/Iris/Private/Iris/ReplicationSystem/DeltaCompression/DeltaCompressionBaselineInvalidationTracker.h b/Engine/Source/Runtime/Net/Iris/Private/Iris/ReplicationSystem/DeltaCompression/DeltaCompressionBaselineInvalidationTracker.h
--- a/Engine/Source/Runtime/Net/Iris/Private/Iris/ReplicationSystem/DeltaCompression/DeltaCompressionBaselineInvalidationTracker.h
+++ b/Engine/Source/Runtime/Net/Iris/Private/Iris/ReplicationSystem/DeltaCompression/DeltaCompressionBaselineInvalidationTracker.h
@@ -9,10 +9,14 @@
 #include "Iris/ReplicationSystem/DeltaCompression/DeltaCompressionBaseline.h"
 #include "Iris/ReplicationSystem/DeltaCompression/DeltaCompressionBaselineStorage.h"
 
+namespace UE::Net
+{
+	typedef uint32 FInternalNetRefIndex;
+}
+
 namespace UE::Net::Private
 {
 	class FDeltaCompressionBaselineManager;
-	typedef uint32 FInternalNetRefIndex;
 }
 
 namespace UE::Net::Private
diff --git a/Engine/Source/Runtime/Net/Iris/Private/Iris/ReplicationSystem/DeltaCompression/DeltaCompressionBaselineManager.cpp b/Engine/Source/Runtime/Net/Iris/Private/Iris/ReplicationSystem/DeltaCompression/DeltaCompressionBaselineManager.cpp
--- a/Engine/Source/Runtime/Net/Iris/Private/Iris/ReplicationSystem/DeltaCompression/DeltaCompressionBaselineManager.cpp
+++ b/Engine/Source/Runtime/Net/Iris/Private/Iris/ReplicationSystem/DeltaCompression/DeltaCompressionBaselineManager.cpp
@@ -55,7 +55,6 @@ FDeltaCompressionBaselineManager::FDeltaCompressionBaselineManager()
 {
 	// If alignment of storage type for FPerObjectInfo is too small then it's totally unsafe to cast such pointers to FPerObjectInfo.
 	static_assert(InvalidObjectInfoIndex == 0, "InvalidObjectInfoIndex needs to be zero in order for current code to function properly.");
-	static_assert(InvalidInternalIndex == FNetRefHandleManager::InvalidInternalIndex, "InvalidInternalIndex differs from FNetRefHandleManager::InvalidInternalIndex");
 }
 
 FDeltaCompressionBaselineManager::~FDeltaCompressionBaselineManager()
diff --git a/Engine/Source/Runtime/Net/Iris/Private/Iris/ReplicationSystem/DeltaCompression/DeltaCompressionBaselineManager.h b/Engine/Source/Runtime/Net/Iris/Private/Iris/ReplicationSystem/DeltaCompression/DeltaCompressionBaselineManager.h
--- a/Engine/Source/Runtime/Net/Iris/Private/Iris/ReplicationSystem/DeltaCompression/DeltaCompressionBaselineManager.h
+++ b/Engine/Source/Runtime/Net/Iris/Private/Iris/ReplicationSystem/DeltaCompression/DeltaCompressionBaselineManager.h
@@ -5,6 +5,7 @@
 #include "Containers/Array.h"
 #include "Net/Core/NetBitArray.h"
 #include "Iris/IrisConfig.h"
+#include "Iris/ReplicationSystem/InternalNetRefIndexManager.h"
 #include "Iris/ReplicationSystem/DeltaCompression/DeltaCompressionBaseline.h"
 #include "Iris/ReplicationSystem/DeltaCompression/DeltaCompressionBaselineStorage.h"
 #include "Misc/TransactionallySafeRWLock.h"
@@ -20,7 +21,6 @@ namespace UE::Net::Private
 {
 	struct FChangeMaskCache;
 	class FDeltaCompressionBaselineInvalidationTracker;
-	typedef uint32 FInternalNetRefIndex;
 	class FNetRefHandleManager;
 	class FReplicationConnections;
 	class FScopedBaselineAccess;
@@ -67,7 +67,7 @@ public:
 	void Deinit();
 
 	/** Called when the maximum InternalNetRefIndex increased and we need to realloc our lists */
-	void OnMaxInternalNetRefIndexIncreased(UE::Net::Private::FInternalNetRefIndex NewMaxInternalIndex);
+	void OnMaxInternalNetRefIndexIncreased(FInternalNetRefIndex NewMaxInternalIndex);
 
 	void PreSendUpdate(FDeltaCompressionBaselineManagerPreSendUpdateParams& UpdateParams);
 	void PostSendUpdate(FDeltaCompressionBaselineManagerPostSendUpdateParams& UpdateParams);
@@ -113,9 +113,6 @@ private:
 		InvalidObjectInfoIndex = 0,
 		InvalidBaselineStateInfoIndex = 0,
 
-		// Must match NetRefHandleManager::InvalidInternalIndex
-		InvalidInternalIndex = 0,
-
 		ObjectInfoGrowCount = 128,
 	};
 
@@ -305,7 +302,7 @@ inline const FDeltaCompressionBaselineManager::FPerObjectInfo* FDeltaCompression
 
 inline FDeltaCompressionBaselineManager::FPerObjectInfo* FDeltaCompressionBaselineManager::GetPerObjectInfoForObject(uint32 ObjectIndex)
 {
-	checkSlow(ObjectIndex != InvalidInternalIndex);
+	checkSlow(ObjectIndex != InvalidInternalNetRefIndex);
 	checkSlow(ObjectIndex < (uint32)ObjectIndexToObjectInfoIndex.Num());
 
 	const ObjectInfoIndexType InfoIndex = ObjectIndexToObjectInfoIndex[ObjectIndex];
@@ -319,7 +316,7 @@ inline FDeltaCompressionBaselineManager::FPerObjectInfo* FDeltaCompressionBaseli
 
 inline const FDeltaCompressionBaselineManager::FPerObjectInfo* FDeltaCompressionBaselineManager::GetPerObjectInfoForObject(uint32 ObjectIndex) const
 {
-	checkSlow(ObjectIndex != InvalidInternalIndex);
+	checkSlow(ObjectIndex != InvalidInternalNetRefIndex);
 	checkSlow(ObjectIndex < (uint32)ObjectIndexToObjectInfoIndex.Num());
 
 	const ObjectInfoIndexType InfoIndex = ObjectIndexToObjectInfoIndex[ObjectIndex];
diff --git a/Engine/Source/Runtime/Net/Iris/Private/Iris/ReplicationSystem/DirtyNetObjectTracker.cpp b/Engine/Source/Runtime/Net/Iris/Private/Iris/ReplicationSystem/DirtyNetObjectTracker.cpp
--- a/Engine/Source/Runtime/Net/Iris/Private/Iris/ReplicationSystem/DirtyNetObjectTracker.cpp
+++ b/Engine/Source/Runtime/Net/Iris/Private/Iris/ReplicationSystem/DirtyNetObjectTracker.cpp
@@ -166,7 +166,7 @@ void FDirtyNetObjectTracker::ApplyGlobalDirtyObjectList()
 				const FNetHandle NetHandle = Element.Key;
 
 				const FInternalNetRefIndex NetObjectIndex = NetRefHandleManager->GetInternalIndexFromNetHandle(NetHandle);
-				if (NetObjectIndex != FNetRefHandleManager::InvalidInternalIndex)
+				if (NetObjectIndex != InvalidInternalNetRefIndex)
 				{
 					// Update changemasks based on rep-indices marked dirty
 					const FGlobalDirtyNetObjectTracker::FDirtyPropertyStorage& DirtyPropertyData = Element.Value;
@@ -197,7 +197,7 @@ void FDirtyNetObjectTracker::ApplyGlobalDirtyObjectList()
 		for (FNetHandle NetHandle : GlobalDirtyNetObjects)
 		{
 			const FInternalNetRefIndex NetObjectIndex = NetRefHandleManager->GetInternalIndexFromNetHandle(NetHandle);
-			if (NetObjectIndex != FNetRefHandleManager::InvalidInternalIndex)
+			if (NetObjectIndex != InvalidInternalNetRefIndex)
 			{
 				DirtyNetObjects.SetBit(NetObjectIndex);
 			}
@@ -274,7 +274,7 @@ void FDirtyNetObjectTracker::MarkNetObjectDirty(FInternalNetRefIndex NetObjectIn
 	checkf(bIsExternalAccessAllowed, TEXT("Cannot mark objects dirty while the bitarray is locked for modifications."));
 #endif
 
-	if (NetObjectIndex >= NetObjectIdCount || NetObjectIndex == FNetRefHandleManager::InvalidInternalIndex)
+	if (NetObjectIndex >= NetObjectIdCount || NetObjectIndex == InvalidInternalNetRefIndex)
 	{
 		UE_LOG(LogIrisDirtyTracker, Warning, TEXT("FDirtyNetObjectTracker::MarkNetObjectDirty received invalid NetObjectIndex: %u | Max: %u"), NetObjectIndex, NetObjectIdCount);
 		return;
diff --git a/Engine/Source/Runtime/Net/Iris/Private/Iris/ReplicationSystem/DirtyNetObjectTracker.h b/Engine/Source/Runtime/Net/Iris/Private/Iris/ReplicationSystem/DirtyNetObjectTracker.h
--- a/Engine/Source/Runtime/Net/Iris/Private/Iris/ReplicationSystem/DirtyNetObjectTracker.h
+++ b/Engine/Source/Runtime/Net/Iris/Private/Iris/ReplicationSystem/DirtyNetObjectTracker.h
@@ -10,12 +10,15 @@
 #include "Iris/IrisConfig.h"
 #include "Iris/Core/IrisCsv.h"
 
+namespace UE::Net
+{
+	typedef uint32 FInternalNetRefIndex;
+}
+
 namespace UE::Net::Private
 {
 	class FNetRefHandleManager;
 	class FDirtyObjectsAccessor;
-	
-	typedef uint32 FInternalNetRefIndex;
 }
 
 #ifndef UE_NET_DIRTYOBJECTTRACKER_LOG_COMPILE_VERBOSITY
diff --git a/Engine/Source/Runtime/Net/Iris/Private/Iris/ReplicationSystem/Filtering/AlwaysRelevantNetObjectFilter.cpp b/Engine/Source/Runtime/Net/Iris/Private/Iris/ReplicationSystem/Filtering/AlwaysRelevantNetObjectFilter.cpp
--- a/Engine/Source/Runtime/Net/Iris/Private/Iris/ReplicationSystem/Filtering/AlwaysRelevantNetObjectFilter.cpp
+++ b/Engine/Source/Runtime/Net/Iris/Private/Iris/ReplicationSystem/Filtering/AlwaysRelevantNetObjectFilter.cpp
@@ -12,12 +12,12 @@ void UAlwaysRelevantNetObjectFilter::OnDeinit()
 {
 }
 
-bool UAlwaysRelevantNetObjectFilter::AddObject(uint32 ObjectIndex, FNetObjectFilterAddObjectParams& Params)
+bool UAlwaysRelevantNetObjectFilter::AddObject(UE::Net::FInternalNetRefIndex ObjectIndex, FNetObjectFilterAddObjectParams& Params)
 {
 	return true;
 }
 
-void UAlwaysRelevantNetObjectFilter::RemoveObject(uint32 ObjectIndex, const FNetObjectFilteringInfo& Info)
+void UAlwaysRelevantNetObjectFilter::RemoveObject(UE::Net::FInternalNetRefIndex ObjectIndex, const FNetObjectFilteringInfo& Info)
 {
 }
 
@@ -31,7 +31,7 @@ void UAlwaysRelevantNetObjectFilter::Filter(FNetObjectFilteringParams& Params)
 	Params.OutAllowedObjects.SetAllBits();
 }
 
-void UAlwaysRelevantNetObjectFilter::OnMaxInternalNetRefIndexIncreased(uint32 NewMaxInternalIndex)
+void UAlwaysRelevantNetObjectFilter::OnMaxInternalNetRefIndexIncreased(UE::Net::FInternalNetRefIndex NewMaxInternalIndex)
 {
 }
 
diff --git a/Engine/Source/Runtime/Net/Iris/Private/Iris/ReplicationSystem/Filtering/NetObjectFilter.cpp b/Engine/Source/Runtime/Net/Iris/Private/Iris/ReplicationSystem/Filtering/NetObjectFilter.cpp
--- a/Engine/Source/Runtime/Net/Iris/Private/Iris/ReplicationSystem/Filtering/NetObjectFilter.cpp
+++ b/Engine/Source/Runtime/Net/Iris/Private/Iris/ReplicationSystem/Filtering/NetObjectFilter.cpp
@@ -56,7 +56,7 @@ void UNetObjectFilter::Deinit()
 	NetRefHandleManager = nullptr;
 }
 
-void UNetObjectFilter::MaxInternalNetRefIndexIncreased(UE::Net::Private::FInternalNetRefIndex MaxInternalIndex, TArrayView<FNetObjectFilteringInfo> NewFilterInfoView)
+void UNetObjectFilter::MaxInternalNetRefIndexIncreased(UE::Net::FInternalNetRefIndex MaxInternalIndex, TArrayView<FNetObjectFilteringInfo> NewFilterInfoView)
 {
 	FilteredObjects.SetNumBits(MaxInternalIndex);
 	
@@ -91,7 +91,7 @@ void UNetObjectFilter::PostFilter(FNetObjectPostFilteringParams&)
 {
 }
 
-FNetObjectFilteringInfo* UNetObjectFilter::GetFilteringInfo(uint32 ObjectIndex)
+FNetObjectFilteringInfo* UNetObjectFilter::GetFilteringInfo(UE::Net::FInternalNetRefIndex ObjectIndex)
 {
 	// Only allow retrieving infos for objects handled by this instance.
 	if (!IsObjectFiltered(ObjectIndex))
@@ -102,7 +102,7 @@ FNetObjectFilteringInfo* UNetObjectFilter::GetFilteringInfo(uint32 ObjectIndex)
 	return &FilteringInfos[ObjectIndex];
 }
 
-uint32 UNetObjectFilter::GetObjectIndex(UE::Net::FNetRefHandle NetRefHandle) const
+UE::Net::FInternalNetRefIndex UNetObjectFilter::GetObjectIndex(UE::Net::FNetRefHandle NetRefHandle) const
 {
 	return NetRefHandleManager->GetInternalIndex(NetRefHandle);
 }
diff --git a/Engine/Source/Runtime/Net/Iris/Private/Iris/ReplicationSystem/Filtering/NetObjectGridFilter.cpp b/Engine/Source/Runtime/Net/Iris/Private/Iris/ReplicationSystem/Filtering/NetObjectGridFilter.cpp
--- a/Engine/Source/Runtime/Net/Iris/Private/Iris/ReplicationSystem/Filtering/NetObjectGridFilter.cpp
+++ b/Engine/Source/Runtime/Net/Iris/Private/Iris/ReplicationSystem/Filtering/NetObjectGridFilter.cpp
@@ -75,10 +75,10 @@ uint16 UNetObjectGridFilter::GetFrameCountBeforeCulling(FName ProfileName) const
 	return FrameCount;
 }
 
-bool UNetObjectGridFilter::AddObject(uint32 ObjectIndex, FNetObjectFilterAddObjectParams& Params)
+bool UNetObjectGridFilter::AddObject(UE::Net::FInternalNetRefIndex ObjectNetIndex, FNetObjectFilterAddObjectParams& Params)
 {
 	// We only support objects that registered a world location
-	const bool bHasLocation = BuildObjectInfo(ObjectIndex, Params);
+	const bool bHasLocation = BuildObjectInfo(ObjectNetIndex, Params);
 
 	if (!bHasLocation)
 	{
@@ -92,32 +92,32 @@ bool UNetObjectGridFilter::AddObject(uint32 ObjectIndex, FNetObjectFilterAddObje
 
 	FPerObjectInfo& PerObjectInfo = ObjectInfos[InfoIndex];
 
-	PerObjectInfo.ObjectIndex = ObjectIndex;
+	PerObjectInfo.ObjectIndex = ObjectNetIndex;
 	PerObjectInfo.FrameCountBeforeCulling = GetFrameCountBeforeCulling(Params.ProfileName);
 
 	if (!AddCellInfoForObject(ObjectLocationInfo, Params.InstanceProtocol))
 	{
 		// Too big an object. We expect it to be costly to move it across cells.
-		UE_LOG(LogIris, Warning, TEXT("ReplicatedObject %s cull distance %f is above the max %f. Object will become always relevant instead"), *NetRefHandleManager->PrintObjectFromIndex(ObjectIndex), PerObjectInfo.GetCullDistance(), MaxNetCullDistance);
-		RemoveObject(ObjectIndex, ObjectLocationInfo);
+		UE_LOG(LogIris, Warning, TEXT("ReplicatedObject %s cull distance %f is above the max %f. Object will become always relevant instead"), *NetRefHandleManager->PrintObjectFromIndex(ObjectNetIndex), PerObjectInfo.GetCullDistance(), MaxNetCullDistance);
+		RemoveObject(ObjectNetIndex, ObjectLocationInfo);
 		return false;
 	}
 
 	return true;
 }
 
-void UNetObjectGridFilter::RemoveObject(uint32 ObjectIndex, const FNetObjectFilteringInfo& Info)
+void UNetObjectGridFilter::RemoveObject(UE::Net::FInternalNetRefIndex ObjectNetIndex, const FNetObjectFilteringInfo& Info)
 {
 	// Remove the object from connection lists
 	for (FPerConnectionInfo& ConnectionInfo : PerConnectionInfos)
 	{
-		ConnectionInfo.RecentObjectFrameCount.Remove(ObjectIndex);
+		ConnectionInfo.RecentObjectFrameCount.Remove(ObjectNetIndex);
 	}
 
 	const FObjectLocationInfo& ObjectLocationInfo = static_cast<const FObjectLocationInfo&>(Info);
 	RemoveCellInfoForObject(ObjectLocationInfo);
 
-	OnObjectRemoved(ObjectIndex);
+	OnObjectRemoved(ObjectNetIndex);
 
 	const uint32 InfoIndex = ObjectLocationInfo.GetInfoIndex();
 	FreeObjectInfo(InfoIndex);
@@ -201,7 +201,7 @@ void UNetObjectGridFilter::Filter(FNetObjectFilteringParams& Params)
 		{
 			if (FCellObjects* Objects = Cells.Find(CellAndTimestamp.Cell))
 			{
-				for (const uint32 ObjectIndex : Objects->ObjectIndices)
+				for (const UE::Net::FInternalNetRefIndex ObjectIndex : Objects->ObjectIndices)
 				{
 					const FObjectLocationInfo& ObjectLocationInfo = static_cast<const FObjectLocationInfo&>(Params.FilteringInfos[ObjectIndex]);
 					const FPerObjectInfo& PerObjectInfo = ObjectInfos[ObjectLocationInfo.GetInfoIndex()];
@@ -308,7 +308,7 @@ bool UNetObjectGridFilter::AddCellInfoForObject(const FObjectLocationInfo& Objec
 		PerObjectInfo.CellBox = NewCellBox;
 	}
 
-	const uint32 ObjectIndex = PerObjectInfo.ObjectIndex;
+	const UE::Net::FInternalNetRefIndex ObjectIndex = PerObjectInfo.ObjectIndex;
 
 	// Add object to new cells.
 	{
@@ -333,7 +333,7 @@ void UNetObjectGridFilter::RemoveCellInfoForObject(const FObjectLocationInfo& Ob
 	FPerObjectInfo& PerObjectInfo = ObjectInfos[ObjectLocationInfo.GetInfoIndex()];
 
 	const FCellBox CellBox = PerObjectInfo.CellBox;
-	const uint32 ObjectIndex = PerObjectInfo.ObjectIndex;
+	const UE::Net::FInternalNetRefIndex ObjectIndex = PerObjectInfo.ObjectIndex;
 
 	// Remove object from cells.
 	{
@@ -367,7 +367,7 @@ void UNetObjectGridFilter::UpdateCellInfoForObject(const FObjectLocationInfo& Ob
 
 	if (NewCellBox != PrevCellBox)
 	{
-		const uint32 ObjectIndex = PerObjectInfo.ObjectIndex;
+		const UE::Net::FInternalNetRefIndex ObjectIndex = PerObjectInfo.ObjectIndex;
 		PerObjectInfo.CellBox = NewCellBox;
 
 		// Special case for disjoint cells
diff --git a/Engine/Source/Runtime/Net/Iris/Private/Iris/ReplicationSystem/Filtering/NetObjectGroups.cpp b/Engine/Source/Runtime/Net/Iris/Private/Iris/ReplicationSystem/Filtering/NetObjectGroups.cpp
--- a/Engine/Source/Runtime/Net/Iris/Private/Iris/ReplicationSystem/Filtering/NetObjectGroups.cpp
+++ b/Engine/Source/Runtime/Net/Iris/Private/Iris/ReplicationSystem/Filtering/NetObjectGroups.cpp
@@ -192,7 +192,7 @@ bool FNetObjectGroups::Contains(FNetObjectGroupHandle GroupHandle, FInternalNetR
 
 void FNetObjectGroups::AddToGroup(FNetObjectGroupHandle GroupHandle, FInternalNetRefIndex InternalIndex)
 {
-	if (InternalIndex == FNetRefHandleManager::InvalidInternalIndex) 
+	if (InternalIndex == InvalidInternalNetRefIndex) 
 	{
 		return;
 	}
@@ -228,7 +228,7 @@ void FNetObjectGroups::AddToGroup(FNetObjectGroupHandle GroupHandle, FInternalNe
 
 void FNetObjectGroups::RemoveFromGroup(FNetObjectGroupHandle GroupHandle, FInternalNetRefIndex InternalIndex)
 {
-	if (InternalIndex == FNetRefHandleManager::InvalidInternalIndex)
+	if (InternalIndex == InvalidInternalNetRefIndex)
 	{
 		return;
 	}
diff --git a/Engine/Source/Runtime/Net/Iris/Private/Iris/ReplicationSystem/Filtering/ObjectScopeHysteresisUpdater.h b/Engine/Source/Runtime/Net/Iris/Private/Iris/ReplicationSystem/Filtering/ObjectScopeHysteresisUpdater.h
--- a/Engine/Source/Runtime/Net/Iris/Private/Iris/ReplicationSystem/Filtering/ObjectScopeHysteresisUpdater.h
+++ b/Engine/Source/Runtime/Net/Iris/Private/Iris/ReplicationSystem/Filtering/ObjectScopeHysteresisUpdater.h
@@ -7,7 +7,7 @@
 #include "Containers/Map.h"
 #include "Net/Core/NetBitArray.h"
 
-namespace UE::Net::Private
+namespace UE::Net
 {
 	typedef uint32 FInternalNetRefIndex;
 }
diff --git a/Engine/Source/Runtime/Net/Iris/Private/Iris/ReplicationSystem/Filtering/ReplicationFiltering.cpp b/Engine/Source/Runtime/Net/Iris/Private/Iris/ReplicationSystem/Filtering/ReplicationFiltering.cpp
--- a/Engine/Source/Runtime/Net/Iris/Private/Iris/ReplicationSystem/Filtering/ReplicationFiltering.cpp
+++ b/Engine/Source/Runtime/Net/Iris/Private/Iris/ReplicationSystem/Filtering/ReplicationFiltering.cpp
@@ -13,6 +13,7 @@
 #include "Iris/Core/IrisLog.h"
 #include "Iris/Core/IrisLogUtils.h"
 #include "Iris/Core/IrisProfiler.h"
+#include "Iris/ReplicationSystem/InternalNetRefIndexManager.h"
 #include "Iris/ReplicationSystem/NetRefHandleManager.h"
 #include "Iris/ReplicationSystem/ReplicationConnections.h"
 #include "Iris/ReplicationSystem/ReplicationProtocol.h"
@@ -97,7 +98,7 @@ public:
 
 	struct FPerFilterInfo
 	{
-		uint32* ObjectIndices = nullptr;
+		UE::Net::FInternalNetRefIndex* ObjectIndices = nullptr;
 		uint32 ObjectCount = 0;
 	};
 
@@ -117,11 +118,11 @@ public:
 		}
 	}
 
-	void PrepareBatch(const uint32* ObjectIndices, uint32 ObjectCount, const TArray<uint8>& FilterIndices)
+	void PrepareBatch(const UE::Net::FInternalNetRefIndex* ObjectIndices, uint32 ObjectCount, const TArray<uint8>& FilterIndices)
 	{
 		ResetBatch();
 
-		for (const uint32 ObjectIndex : MakeArrayView(ObjectIndices, ObjectCount))
+		for (const UE::Net::FInternalNetRefIndex ObjectIndex : MakeArrayView(ObjectIndices, ObjectCount))
 		{
 			const uint8 FilterIndex = FilterIndices[ObjectIndex];
 			if (FilterIndex == InvalidDynamicFilterIndex)
@@ -495,7 +496,7 @@ bool FReplicationFiltering::SetFilter(FInternalNetRefIndex ObjectIndex, FNetObje
 	const FNetRefHandleManager::FReplicatedObjectData& ObjectData = NetRefHandleManager->GetReplicatedObjectDataNoCheck(ObjectIndex);
 	const TNetChunkedArray<uint8*>& ReplicatedObjectsStateBuffers = NetRefHandleManager->GetReplicatedObjectStateBuffers();
 	// Let subobjects be filtered like their owners.
-	if (bWantsToUseDynamicFilter && (ObjectData.SubObjectRootIndex != FNetRefHandleManager::InvalidInternalIndex))
+	if (bWantsToUseDynamicFilter && (ObjectData.SubObjectRootIndex != InvalidInternalNetRefIndex))
 	{
 		UE_LOG(LogIrisFiltering, Warning, TEXT("Cannot set dynamic filters on subobjects. Filter change for %s is ignored"), *NetRefHandleManager->PrintObjectFromIndex(ObjectIndex));
 		return false;
@@ -510,11 +511,23 @@ bool FReplicationFiltering::SetFilter(FInternalNetRefIndex ObjectIndex, FNetObje
 		this->FreePerObjectInfoForObject(ObjIndex);
 	};
 
-	auto TrySetDynamicFilter = [this, &ObjectData, &ReplicatedObjectsStateBuffers, FilterConfigProfile](uint32 ObjIndex, uint32 FilterIndex)
+	const FInternalNetRefIndexManager NetRefIndexManager(*NetRefHandleManager);
+
+	auto TrySetDynamicFilter = [this, &ObjectData, &ReplicatedObjectsStateBuffers, FilterConfigProfile, NetRefIndexManager](uint32 ObjIndex, uint32 FilterIndex)
 	{
 		FNetObjectFilteringInfo& NetObjectFilteringInfo = this->NetObjectFilteringInfos[ObjIndex];
 		NetObjectFilteringInfo = {};
-		FNetObjectFilterAddObjectParams AddParams = { .OutInfo=NetObjectFilteringInfo, .ProfileName=FilterConfigProfile, .InstanceProtocol=ObjectData.InstanceProtocol, .Protocol=ObjectData.Protocol, .StateBuffer=ReplicatedObjectsStateBuffers[ObjIndex] };
+		
+		FNetObjectFilterAddObjectParams AddParams
+		{ 
+			.OutInfo = NetObjectFilteringInfo,
+			.ProfileName = FilterConfigProfile,
+			.InstanceProtocol = ObjectData.InstanceProtocol,
+			.Protocol = ObjectData.Protocol,
+			.NetRefIndexManager = NetRefIndexManager,
+			.StateBuffer = ReplicatedObjectsStateBuffers[ObjIndex] 
+		};
+
 		FFilterInfo& FilterInfo = this->DynamicFilterInfos[FilterIndex];
 		if (FilterInfo.Filter->AddObject(ObjIndex, AddParams))
 		{
@@ -705,7 +718,7 @@ void FReplicationFiltering::InitNewConnections()
 		this->SetPerConnectionListsSize(ConnectionInfo, CurrentMaxInternalindex);
 
 		ConnectionInfo.ConnectionFilteredObjects.Copy(ScopableInternalIndices);
-		ConnectionInfo.ConnectionFilteredObjects.ClearBit(FNetRefHandleManager::InvalidInternalIndex);
+		ConnectionInfo.ConnectionFilteredObjects.ClearBit(InvalidInternalNetRefIndex);
 
 		// Update group exclusion filtering
 		{
@@ -928,7 +941,7 @@ void FReplicationFiltering::UpdateObjectsInScope()
 		uint32* ObjectsWithOwnerFilterStorage = ObjectsWithOwnerFilter.GetDataChecked(WordCountForObjectBitArrays);
 		uint32* ObjectsExemptFromHysteresisStorage = HysteresisState.ObjectsExemptFromHysteresis.GetDataChecked(WordCountForObjectBitArrays);
 		const FNetBitArrayView SubObjectInternalIndices = NetRefHandleManager->GetSubObjectInternalIndicesView();
-		uint32 PrevParentIndex = FNetRefHandleManager::InvalidInternalIndex;
+		uint32 PrevParentIndex = InvalidInternalNetRefIndex;
 		for (uint32 WordIt = 0, WordEndIt = ModifiedWordIndex; WordIt != WordEndIt; ++WordIt)
 		{
 			const uint32 WordIndex = ModifiedWords[WordIt];
@@ -978,7 +991,7 @@ void FReplicationFiltering::UpdateObjectsInScope()
 					const uint32 ObjectIndex = BitOffset + FPlatformMath::CountTrailingZeros(LeastSignificantBit);
 					const uint32 ParentIndex = NetRefHandleManager->GetReplicatedObjectDataNoCheck(ObjectIndex).SubObjectRootIndex;
 					
-					if (ParentIndex == PrevParentIndex || ParentIndex == FNetRefHandleManager::InvalidInternalIndex)
+					if (ParentIndex == PrevParentIndex || ParentIndex == InvalidInternalNetRefIndex)
 					{
 						continue;
 					}
@@ -1276,6 +1289,8 @@ void FReplicationFiltering::PreUpdateDynamicFiltering()
 {
 	IRIS_PROFILER_SCOPE(FReplicationFiltering_PreUpdateDynamicFiltering);
 
+	const FInternalNetRefIndexManager NetRefIndexManager(*NetRefHandleManager);
+
 	// Give filters a chance to prepare for filtering. It's only called if any object has the filter set.
 	{
 		for (FFilterInfo& Info : DynamicFilterInfos)
@@ -1285,9 +1300,12 @@ void FReplicationFiltering::PreUpdateDynamicFiltering()
 				continue;
 			}
 
-			FNetObjectPreFilteringParams PreFilteringParams;
-			PreFilteringParams.FilteringInfos = MakeArrayView(NetObjectFilteringInfos);
-			PreFilteringParams.ValidConnections = MakeNetBitArrayView(ValidConnections);
+			FNetObjectPreFilteringParams PreFilteringParams
+			{
+				.ValidConnections = MakeNetBitArrayView(ValidConnections),
+				.FilteringInfos = MakeArrayView(NetObjectFilteringInfos),
+				.NetRefIndexManager = NetRefIndexManager,
+			};
 			Info.Filter->PreFilter(PreFilteringParams);
 		}
 	}
@@ -1508,6 +1526,8 @@ void FReplicationFiltering::UpdateDynamicFiltering()
 	TArray<uint32> FilteredOutByHysteresisObjects;
 	FilteredOutByHysteresisObjects.Reserve(256);
 
+	const FInternalNetRefIndexManager NetRefIndexManager(*NetRefHandleManager);
+
 	// For circular dependency detection/avoidance when retrieving dependent object filter status.
 	TSet<FInternalNetRefIndex> VisitedDependentObjects;
 	VisitedDependentObjects.Reserve(32);
@@ -1570,7 +1590,8 @@ void FReplicationFiltering::UpdateDynamicFiltering()
 					.StateBuffers = nullptr,
 					.ConnectionId = ConnId,
 					.View = Connections->GetReplicationView(ConnId),
-					.GroupFilteredOutObjects = ReplicationSystem->GetReplicationSystemInternal()->GetFiltering().GetGroupFilteredOutObjects(ConnId)
+					.GroupFilteredOutObjects = ReplicationSystem->GetReplicationSystemInternal()->GetFiltering().GetGroupFilteredOutObjects(ConnId),
+					.NetRefIndexManager = NetRefIndexManager,
 				};
 
 				// Execute the filter here
@@ -2000,8 +2021,10 @@ void FReplicationFiltering::NotifyFiltersOfDirtyObjects()
 	FUpdateDirtyObjectsBatchHelper BatchHelper(NetRefHandleManager, DynamicFilterInfos);
 
 	constexpr SIZE_T MaxBatchObjectCount = FUpdateDirtyObjectsBatchHelper::Constants::MaxObjectCountPerBatch;
-	uint32 ObjectIndices[MaxBatchObjectCount];
+	UE::Net::FInternalNetRefIndex ObjectIndices[MaxBatchObjectCount];
 
+	// Find up to 512 dirty objects to update every iteration.
+	// Their InternalIndex will get added to ObjectIndices and passed to the filter.
 	const uint32 BitCount = ~0U;
 	for (uint32 ObjectCount, StartIndex = 0; (ObjectCount = DirtyObjectsThisFrame.GetSetBitIndices(StartIndex, BitCount, ObjectIndices, MaxBatchObjectCount)) > 0; )
 	{
@@ -2015,7 +2038,7 @@ void FReplicationFiltering::NotifyFiltersOfDirtyObjects()
 	}
 }
 
-void FReplicationFiltering::BatchNotifyFiltersOfDirtyObjects(FUpdateDirtyObjectsBatchHelper& BatchHelper, const uint32* DirtyObjectIndices, uint32 ObjectCount)
+void FReplicationFiltering::BatchNotifyFiltersOfDirtyObjects(FUpdateDirtyObjectsBatchHelper& BatchHelper, const UE::Net::FInternalNetRefIndex* DirtyObjectIndices, uint32 ObjectCount)
 {
 	BatchHelper.PrepareBatch(DirtyObjectIndices, ObjectCount, ObjectIndexToDynamicFilterIndex);
 
@@ -3069,7 +3092,7 @@ bool FReplicationFiltering::IsValidDynamicFilter(uint32 FilterIndex) const
 		&& DynamicFilterInfos[FilterIndex].Filter != nullptr;
 }
 
-void FReplicationFiltering::RemoveFromDynamicFilter(uint32 ObjectIndex, uint32 FilterIndex)
+void FReplicationFiltering::RemoveFromDynamicFilter(UE::Net::FInternalNetRefIndex ObjectIndex, uint32 FilterIndex)
 {
 	UE_LOG(LogIrisFiltering, Verbose, TEXT("RemoveFromDynamicFilter removing %s from Dynamic Filter %s"), 
 		*NetRefHandleManager->PrintObjectFromIndex(ObjectIndex), *GetFilterName(FNetObjectFilterHandleUtil::MakeDynamicFilterHandle(ObjectIndexToDynamicFilterIndex[ObjectIndex])).ToString());
@@ -3260,7 +3283,7 @@ bool FReplicationFiltering::HasSubObjectInScopeWithFilteredOutRootObject(FNetBit
 	Objects.ForAllSetBits([this, &Objects, &bReturnValue](uint32 ObjectIndex)
 		{
 			const FNetRefHandleManager::FReplicatedObjectData& ReplicatedObjectData = this->NetRefHandleManager->GetReplicatedObjectDataNoCheck(ObjectIndex);
-			if (ReplicatedObjectData.SubObjectRootIndex != FNetRefHandleManager::InvalidInternalIndex)
+			if (ReplicatedObjectData.SubObjectRootIndex != InvalidInternalNetRefIndex)
 			{
 				if (!Objects.GetBit(ReplicatedObjectData.SubObjectRootIndex))
 				{
@@ -3268,7 +3291,7 @@ bool FReplicationFiltering::HasSubObjectInScopeWithFilteredOutRootObject(FNetBit
 					ensureMsgf(Objects.GetBit(ReplicatedObjectData.SubObjectRootIndex), TEXT("Root index %u is not in scope for subobject %u"), ReplicatedObjectData.SubObjectRootIndex, ObjectIndex);
 				}
 			}
-			if (ReplicatedObjectData.SubObjectParentIndex != FNetRefHandleManager::InvalidInternalIndex)
+			if (ReplicatedObjectData.SubObjectParentIndex != InvalidInternalNetRefIndex)
 			{
 				if (!Objects.GetBit(ReplicatedObjectData.SubObjectParentIndex))
 				{
diff --git a/Engine/Source/Runtime/Net/Iris/Private/Iris/ReplicationSystem/Filtering/ReplicationFiltering.h b/Engine/Source/Runtime/Net/Iris/Private/Iris/ReplicationSystem/Filtering/ReplicationFiltering.h
--- a/Engine/Source/Runtime/Net/Iris/Private/Iris/ReplicationSystem/Filtering/ReplicationFiltering.h
+++ b/Engine/Source/Runtime/Net/Iris/Private/Iris/ReplicationSystem/Filtering/ReplicationFiltering.h
@@ -21,12 +21,13 @@ struct FNetObjectFilterDefinition;
 namespace UE::Net
 {
 	typedef uint32 FNetObjectFilterHandle;
+	typedef uint32 FInternalNetRefIndex;
+
 	namespace Private
 	{
 		class FNetRefHandleManager;
 		class FNetObjectGroups;
 		class FReplicationConnections;
-		typedef uint32 FInternalNetRefIndex;
 	}
 }
 
@@ -306,10 +307,10 @@ private:
 	uint32 GetOwningConnectionIfDirty(uint32 ObjectIndex) const;
 
 	bool IsValidDynamicFilter(uint32 FilterIndex) const;
-	void RemoveFromDynamicFilter(uint32 ObjectIndex, uint32 FilterIndex);
+	void RemoveFromDynamicFilter(UE::Net::FInternalNetRefIndex ObjectIndex, uint32 FilterIndex);
 
 	void NotifyFiltersOfDirtyObjects();
-	void BatchNotifyFiltersOfDirtyObjects(FUpdateDirtyObjectsBatchHelper& BatchHelper, const uint32* ObjectIndices, uint32 ObjectCount);
+	void BatchNotifyFiltersOfDirtyObjects(FUpdateDirtyObjectsBatchHelper& BatchHelper, const UE::Net::FInternalNetRefIndex* ObjectIndices, uint32 ObjectCount);
 
 	/** Returns all the filtering infos. */
 	TArrayView<FNetObjectFilteringInfo> GetNetObjectFilteringInfos();
diff --git a/Engine/Source/Runtime/Net/Iris/Private/Iris/ReplicationSystem/InternalNetRefIndexManager.cpp b/Engine/Source/Runtime/Net/Iris/Private/Iris/ReplicationSystem/InternalNetRefIndexManager.cpp
new file mode 100644
--- /dev/null
+++ b/Engine/Source/Runtime/Net/Iris/Private/Iris/ReplicationSystem/InternalNetRefIndexManager.cpp
@@ -0,0 +1,17 @@
+// Copyright Epic Games, Inc. All Rights Reserved.
+
+#include "Iris/ReplicationSystem/InternalNetRefIndexManager.h"
+
+#include "Iris/ReplicationSystem/NetRefHandleManager.h"
+
+namespace UE::Net
+{
+
+FInternalNetRefIndexManager::FInternalNetRefIndexManager(const UE::Net::Private::FNetRefHandleManager& InNetRefHandleManager)
+	: NetRefHandleManager(InNetRefHandleManager)
+{
+
+}
+
+
+} // end namespace UE::Net
\ No newline at end of file
diff --git a/Engine/Source/Runtime/Net/Iris/Private/Iris/ReplicationSystem/NetBlob/NetBlobManager.cpp b/Engine/Source/Runtime/Net/Iris/Private/Iris/ReplicationSystem/NetBlob/NetBlobManager.cpp
--- a/Engine/Source/Runtime/Net/Iris/Private/Iris/ReplicationSystem/NetBlob/NetBlobManager.cpp
+++ b/Engine/Source/Runtime/Net/Iris/Private/Iris/ReplicationSystem/NetBlob/NetBlobManager.cpp
@@ -1,17 +1,23 @@
 // Copyright Epic Games, Inc. All Rights Reserved.
 
 #include "Iris/ReplicationSystem/NetBlob/NetBlobManager.h"
+
 #include "HAL/IConsoleManager.h"
+
 #include "Iris/ReplicationSystem/NetBlob/NetObjectBlobHandler.h"
 #include "Iris/ReplicationSystem/NetBlob/NetRPCHandler.h"
 #include "Iris/ReplicationSystem/NetBlob/PartialNetObjectAttachmentHandler.h"
+
 #include "Iris/ReplicationSystem/ReplicationConnections.h"
 #include "Iris/ReplicationSystem/ReplicationSystem.h"
 #include "Iris/ReplicationSystem/ReplicationSystemInternal.h"
 #include "Iris/ReplicationSystem/ReplicationWriter.h"
+#include "Iris/ReplicationSystem/NetRefHandleManager.h"
 #include "Iris/ReplicationSystem/ObjectReplicationBridge.h"
+
 #include "Iris/Core/NetObjectReference.h"
 #include "Iris/Core/IrisLog.h"
+
 #include "Net/Core/Trace/NetDebugName.h"
 
 namespace UE::Net::Private
@@ -320,12 +326,12 @@ bool FNetBlobManager::GetRootObjectIndicesFromHandle(FNetRefHandle RootObjectRef
 	}
 
 	const FInternalNetRefIndex ObjectIndex = NetRefHandleManager->GetInternalIndex(RootObjectRefHandle);
-	if (ObjectIndex == FNetRefHandleManager::InvalidInternalIndex)
+	if (ObjectIndex == InvalidInternalNetRefIndex)
 	{
 		return false;
 	}
 
-	checkf(NetRefHandleManager->GetReplicatedObjectDataNoCheck(ObjectIndex).SubObjectRootIndex == FNetRefHandleManager::InvalidInternalIndex, TEXT("Object %s (index:%u) (netref:%s) is not a rootobject"),
+	checkf(NetRefHandleManager->GetReplicatedObjectDataNoCheck(ObjectIndex).SubObjectRootIndex == InvalidInternalNetRefIndex, TEXT("Object %s (index:%u) (netref:%s) is not a rootobject"),
 		*GetNameSafe(NetRefHandleManager->GetReplicatedObjectInstance(ObjectIndex)), ObjectIndex, ToCStr(RootObjectRefHandle.ToString()));
 
 	OutRootObjectIndex = ObjectIndex;
@@ -341,20 +347,20 @@ bool FNetBlobManager::GetRootObjectAndSubObjectIndicesFromSubObjectHandle(FNetRe
 	}
 
 	const FInternalNetRefIndex ObjectIndex = NetRefHandleManager->GetInternalIndex(SubObjectRefHandle);
-	if (ObjectIndex == FNetRefHandleManager::InvalidInternalIndex)
+	if (ObjectIndex == InvalidInternalNetRefIndex)
 	{
 		return false;
 	}
 
 	const FNetRefHandleManager::FReplicatedObjectData& ObjectData = NetRefHandleManager->GetReplicatedObjectDataNoCheck(ObjectIndex);
 
-	checkf(ObjectData.SubObjectRootIndex != FNetRefHandleManager::InvalidInternalIndex, TEXT("SubObject %s (index:%u) (netref:%s) does not have a rootobject"),
+	checkf(ObjectData.SubObjectRootIndex != InvalidInternalNetRefIndex, TEXT("SubObject %s (index:%u) (netref:%s) does not have a rootobject"),
 		*GetNameSafe(NetRefHandleManager->GetReplicatedObjectInstance(ObjectIndex)), ObjectIndex, ToCStr(SubObjectRefHandle.ToString()));
 
 	OutRootObjectIndex = ObjectData.SubObjectRootIndex;
 	OutSubObjectIndex = ObjectIndex;
 
-	return OutRootObjectIndex != FNetRefHandleManager::InvalidInternalIndex;
+	return OutRootObjectIndex != InvalidInternalNetRefIndex;
 }
 
 bool FNetBlobManager::GetRootObjectAndSubObjectIndicesFromAnyHandle(FNetRefHandle AnyRefHandle, FInternalNetRefIndex& OutRootObjectIndex, FInternalNetRefIndex& OutSubObjectIndex) const
@@ -365,14 +371,14 @@ bool FNetBlobManager::GetRootObjectAndSubObjectIndicesFromAnyHandle(FNetRefHandl
 	}
 
 	const FInternalNetRefIndex ObjectIndex = NetRefHandleManager->GetInternalIndex(AnyRefHandle);
-	if (ObjectIndex == FNetRefHandleManager::InvalidInternalIndex)
+	if (ObjectIndex == InvalidInternalNetRefIndex)
 	{
 		return false;
 	}
 
 	const FNetRefHandleManager::FReplicatedObjectData& ObjectData = NetRefHandleManager->GetReplicatedObjectDataNoCheck(ObjectIndex);
 	// If it's a sub object
-	if (ObjectData.SubObjectRootIndex != FNetRefHandleManager::InvalidInternalIndex)
+	if (ObjectData.SubObjectRootIndex != InvalidInternalNetRefIndex)
 	{
 		OutRootObjectIndex = ObjectData.SubObjectRootIndex;
 		OutSubObjectIndex = ObjectIndex;
@@ -381,7 +387,7 @@ bool FNetBlobManager::GetRootObjectAndSubObjectIndicesFromAnyHandle(FNetRefHandl
 	else
 	{
 		OutRootObjectIndex = ObjectIndex;
-		OutSubObjectIndex = FNetRefHandleManager::InvalidInternalIndex;
+		OutSubObjectIndex = InvalidInternalNetRefIndex;
 	}
 
 	return true;
@@ -594,7 +600,7 @@ void FNetBlobManager::FNetObjectAttachmentSendQueue::PrepareProcessQueue(FReplic
 		uint32 CurrentEntryIndex = 0U;
 		for (const FNetObjectAttachmentQueueEntry& Entry : MakeArrayView(AttachmentQueue))
 		{
-			const uint32 TargetInternalObjectIndex = Entry.SubObjectIndex != FNetRefHandleManager::InvalidInternalIndex ? Entry.SubObjectIndex : Entry.OwnerIndex;
+			const uint32 TargetInternalObjectIndex = Entry.SubObjectIndex != InvalidInternalNetRefIndex ? Entry.SubObjectIndex : Entry.OwnerIndex;
 
 			const bool bIsAttachmentToObjectGoingOutOfScope = !ScopableObjects.GetBit(TargetInternalObjectIndex) && PrevScopableObjects.GetBit(TargetInternalObjectIndex);
 			if (bIsAttachmentToObjectGoingOutOfScope)
diff --git a/Engine/Source/Runtime/Net/Iris/Private/Iris/ReplicationSystem/NetBlob/NetBlobManager.h b/Engine/Source/Runtime/Net/Iris/Private/Iris/ReplicationSystem/NetBlob/NetBlobManager.h
--- a/Engine/Source/Runtime/Net/Iris/Private/Iris/ReplicationSystem/NetBlob/NetBlobManager.h
+++ b/Engine/Source/Runtime/Net/Iris/Private/Iris/ReplicationSystem/NetBlob/NetBlobManager.h
@@ -2,10 +2,15 @@
 
 #pragma once
 #include "CoreTypes.h"
+
+#include "Iris/ReplicationSystem/InternalNetRefIndexManager.h"
 #include "Iris/ReplicationSystem/ReplicationSystemTypes.h"
+
 #include "Iris/ReplicationSystem/NetBlob/NetBlob.h"
 #include "Iris/ReplicationSystem/NetBlob/NetBlobHandlerManager.h"
-#include "Iris/ReplicationSystem/NetRefHandleManager.h" // For FInternalNetRefIndex
+
+#include "Net/Core/NetBitArray.h"
+
 #include "UObject/StrongObjectPtr.h"
 
 class UNetBlobHandler;
@@ -13,10 +18,13 @@ class UNetObjectBlobHandler;
 class UPartialNetObjectAttachmentHandler;
 class UPartialNetObjectAttachmentHandlerConfig;
 class UNetRPCHandler;
+
 namespace UE::Net
 {
 	class FNetObjectReference;
 
+	typedef uint32 FInternalNetRefIndex;
+
 	namespace Private
 	{
 		class FNetRefHandleManager;
@@ -97,8 +105,8 @@ private:
 		// The object the RPC will be applied to.
 		FNetObjectReference TargetRef;
 
-		FInternalNetRefIndex RootObjectIndex = FNetRefHandleManager::InvalidInternalIndex;
-		FInternalNetRefIndex SubObjectIndex = FNetRefHandleManager::InvalidInternalIndex;
+		FInternalNetRefIndex RootObjectIndex = InvalidInternalNetRefIndex;
+		FInternalNetRefIndex SubObjectIndex = InvalidInternalNetRefIndex;
 	};
 	bool GetRPCOwner(FRPCOwner& OutOwnerInfo, const FSendRPCContext& Context) const;
 
diff --git a/Engine/Source/Runtime/Net/Iris/Private/Iris/ReplicationSystem/NetBlob/NetRPC.cpp b/Engine/Source/Runtime/Net/Iris/Private/Iris/ReplicationSystem/NetBlob/NetRPC.cpp
--- a/Engine/Source/Runtime/Net/Iris/Private/Iris/ReplicationSystem/NetBlob/NetRPC.cpp
+++ b/Engine/Source/Runtime/Net/Iris/Private/Iris/ReplicationSystem/NetBlob/NetRPC.cpp
@@ -3,27 +3,34 @@
 #include "Iris/ReplicationSystem/NetBlob/NetRPC.h"
 
 #include "HAL/IConsoleManager.h"
-#include "Iris/Core/BitTwiddling.h"
+
 #include "Net/Core/Trace/NetDebugName.h"
 #include "Net/Core/Trace/NetTrace.h"
 #include "Net/Core/Misc/NetContext.h"
+
+#include "Iris/Core/BitTwiddling.h"
+#include "Iris/Core/IrisLog.h"
+#include "Iris/Core/IrisProfiler.h"
+
 #include "Iris/ReplicationSystem/NetBlob/NetRPCHandler.h"
+
+#include "Iris/ReplicationSystem/InternalNetRefIndexManager.h"
 #include "Iris/ReplicationSystem/ObjectReferenceCache.h"
 #include "Iris/ReplicationSystem/ObjectReplicationBridge.h"
-#include "Iris/Serialization/NetReferenceCollector.h"
 #include "Iris/ReplicationSystem/ReplicationOperations.h"
 #include "Iris/ReplicationSystem/ReplicationOperationsInternal.h"
 #include "Iris/ReplicationSystem/ReplicationProtocol.h"
 #include "Iris/ReplicationSystem/ReplicationSystem.h"
 #include "Iris/ReplicationSystem/ReplicationSystemInternal.h"
+
 #include "Iris/Serialization/InternalNetSerializationContext.h"
+#include "Iris/Serialization/NetBitStreamUtil.h"
 #include "Iris/Serialization/NetBitStreamReader.h"
 #include "Iris/Serialization/NetBitStreamWriter.h"
+#include "Iris/Serialization/NetReferenceCollector.h"
 #include "Iris/Serialization/NetSerializationContext.h"
 #include "Iris/Serialization/ObjectNetSerializer.h"
-#include "Iris/Serialization/NetBitStreamUtil.h"
-#include "Iris/Core/IrisLog.h"
-#include "Iris/Core/IrisProfiler.h"
+
 #include "Containers/ArrayView.h"
 #include "ProfilingDebugging/CsvProfiler.h"
 #include "UObject/CoreNetContext.h"
@@ -520,7 +527,7 @@ void FNetRPC::CallFunction(FNetRPCCallContext& CallContext)
 		const FNetRefHandle RootObjectNetRefHandle = NetObjectReference.GetRefHandle();
 		const FInternalNetRefIndex InternalIndex = NetRefHandleManager.GetInternalIndex(RootObjectNetRefHandle);
 
-		ensureMsgf(InternalIndex != FNetRefHandleManager::InvalidInternalIndex, TEXT("Processing NetRPC for an object before it's internal index was assigned. NetRefHandle: %s | Instance: %s"),
+		ensureMsgf(InternalIndex != InvalidInternalNetRefIndex, TEXT("Processing NetRPC for an object before it's internal index was assigned. NetRefHandle: %s | Instance: %s"),
 			*RootObjectNetRefHandle.ToString(), *GetNameSafe(Object));
 
 		// If the object stopped replicating locally, discard all received RPCs so they don't get executed.
@@ -720,13 +727,13 @@ static UObject* NetRPC_GetRootObject(FNetSerializationContext& Context, const FN
 	
 	FNetRefHandle OwnerRefHandle = ObjectReference.GetRefHandle();
 	const FInternalNetRefIndex InternalIndex = NetRefHandleManager.GetInternalIndex(OwnerRefHandle);
-	if (!ensureMsgf(InternalIndex != FNetRefHandleManager::InvalidInternalIndex, TEXT("Unable to find InternalIndex for object reference %s"), ToCStr(ObjectReference.ToString())))
+	if (!ensureMsgf(InternalIndex != InvalidInternalNetRefIndex, TEXT("Unable to find InternalIndex for object reference %s"), ToCStr(ObjectReference.ToString())))
 	{
 		return nullptr;
 	}
 
 	const FNetRefHandleManager::FReplicatedObjectData& ObjectData = NetRefHandleManager.GetReplicatedObjectDataNoCheck(InternalIndex);
-	if (ObjectData.SubObjectRootIndex != FNetRefHandleManager::InvalidInternalIndex)
+	if (ObjectData.SubObjectRootIndex != InvalidInternalNetRefIndex)
 	{
 		const FNetRefHandleManager::FReplicatedObjectData& RootObjectData = NetRefHandleManager.GetReplicatedObjectDataNoCheck(ObjectData.SubObjectRootIndex);
 		OwnerRefHandle = RootObjectData.RefHandle;
@@ -812,4 +819,4 @@ static FString NetRPC_GetDebugObjectRefName(FNetSerializationContext& Context, c
 	return FString::Printf(TEXT("%s (%s)"), *GetNameSafe(ObjectPtr), *ObjectReference.ToString());
 }
 
-}
+} // end namespace UE::Net::Private
diff --git a/Engine/Source/Runtime/Net/Iris/Private/Iris/ReplicationSystem/NetDependencyData.cpp b/Engine/Source/Runtime/Net/Iris/Private/Iris/ReplicationSystem/NetDependencyData.cpp
--- a/Engine/Source/Runtime/Net/Iris/Private/Iris/ReplicationSystem/NetDependencyData.cpp
+++ b/Engine/Source/Runtime/Net/Iris/Private/Iris/ReplicationSystem/NetDependencyData.cpp
@@ -106,7 +106,7 @@ FNetDependencyData::FSubObjectConditionalsArray& FNetDependencyData::GetOrCreate
 	}
 }
 
-FNetDependencyData::FInternalNetRefIndexArray& FNetDependencyData::GetOrCreateInternalChildSubObjectsArray(FInternalNetRefIndex OwnerIndex, FSubObjectConditionalsArray*& OutSubObjectConditionals)
+FNetDependencyData::FInternalNetIndexArray& FNetDependencyData::GetOrCreateInternalChildSubObjectsArray(FInternalNetRefIndex OwnerIndex, FSubObjectConditionalsArray*& OutSubObjectConditionals)
 {
 	FDependencyInfo& Entry = GetOrCreateCacheEntry(OwnerIndex);
 	if (Entry.ArrayIndices[ChildSubObjects] == FDependencyInfo::InvalidCacheIndex)
@@ -115,7 +115,7 @@ FNetDependencyData::FInternalNetRefIndexArray& FNetDependencyData::GetOrCreateIn
 		FSparseArrayAllocationInfo AllocInfo = DependentObjectsStorage.AddUninitialized();
 		Entry.ArrayIndices[ChildSubObjects] = AllocInfo.Index;
 
-		FInternalNetRefIndexArray* InternalIndexArray = new (AllocInfo.Pointer) FInternalNetRefIndexArray();
+		FInternalNetIndexArray* InternalIndexArray = new (AllocInfo.Pointer) FInternalNetIndexArray();
 		OutSubObjectConditionals = nullptr;
 		return *InternalIndexArray;
 	}
@@ -188,7 +188,7 @@ void FNetDependencyData::FreeCreationDependencyInfoArray(FInternalNetRefIndex Ch
 	}
 }
 
-FNetDependencyData::FInternalNetRefIndexArray& FNetDependencyData::GetOrCreateInternalIndexArray(FInternalNetRefIndex OwnerIndex, EArrayType ArrayType)
+FNetDependencyData::FInternalNetIndexArray& FNetDependencyData::GetOrCreateInternalIndexArray(FInternalNetRefIndex OwnerIndex, EArrayType ArrayType)
 {
 	FDependencyInfo& Entry = GetOrCreateCacheEntry(OwnerIndex);
 	if (Entry.ArrayIndices[ArrayType] == FDependencyInfo::InvalidCacheIndex)
@@ -197,7 +197,7 @@ FNetDependencyData::FInternalNetRefIndexArray& FNetDependencyData::GetOrCreateIn
 		FSparseArrayAllocationInfo AllocInfo = DependentObjectsStorage.AddUninitialized();
 		Entry.ArrayIndices[ArrayType] = AllocInfo.Index;
 
-		FInternalNetRefIndexArray* InternalIndexArray = new (AllocInfo.Pointer) FInternalNetRefIndexArray();
+		FInternalNetIndexArray* InternalIndexArray = new (AllocInfo.Pointer) FInternalNetIndexArray();
 
 		return *InternalIndexArray;
 	}
diff --git a/Engine/Source/Runtime/Net/Iris/Private/Iris/ReplicationSystem/NetDependencyData.h b/Engine/Source/Runtime/Net/Iris/Private/Iris/ReplicationSystem/NetDependencyData.h
--- a/Engine/Source/Runtime/Net/Iris/Private/Iris/ReplicationSystem/NetDependencyData.h
+++ b/Engine/Source/Runtime/Net/Iris/Private/Iris/ReplicationSystem/NetDependencyData.h
@@ -8,9 +8,13 @@
 #include "Containers/SparseArray.h"
 #include "Iris/ReplicationSystem/ReplicationSystemTypes.h"
 
-namespace UE::Net::Private
+namespace UE::Net
 {
 	typedef uint32 FInternalNetRefIndex;
+}
+
+namespace UE::Net::Private
+{
 	typedef int8 FLifeTimeConditionStorage;
 }
 
@@ -36,7 +40,7 @@ class FNetDependencyData
 public:
 	FNetDependencyData();
 
-	typedef TArray<FInternalNetRefIndex, TInlineAllocator<8>> FInternalNetRefIndexArray;
+	typedef TArray<FInternalNetRefIndex, TInlineAllocator<8>> FInternalNetIndexArray;
 	typedef TArray<FLifeTimeConditionStorage, TInlineAllocator<8>> FSubObjectConditionalsArray;
 	typedef TArray<FDependentObjectInfo, TInlineAllocator<8>> FDependentObjectInfoArray;
 	typedef TArray<FInternalNetRefIndex, TInlineAllocator<8>> FCreationDependentInfoArray;
@@ -52,7 +56,7 @@ public:
 	};
 
 	template<EArrayType TypeIndex>
-	FInternalNetRefIndexArray& GetOrCreateInternalIndexArray(FInternalNetRefIndex InternalIndex)
+	FInternalNetIndexArray& GetOrCreateInternalIndexArray(FInternalNetRefIndex InternalIndex)
 	{
 		static_assert(TypeIndex != EArrayType::Count, "Invalid array type index");
 		return GetOrCreateInternalIndexArray(InternalIndex, TypeIndex);
@@ -83,9 +87,9 @@ public:
 
 	FSubObjectConditionalsArray& GetOrCreateSubObjectConditionalsArray(FInternalNetRefIndex InternalIndex);
 
-	FInternalNetRefIndexArray& GetOrCreateInternalChildSubObjectsArray(FInternalNetRefIndex InternalIndex, FSubObjectConditionalsArray*& OutSubObjectConditionals);
+	FInternalNetIndexArray& GetOrCreateInternalChildSubObjectsArray(FInternalNetRefIndex InternalIndex, FSubObjectConditionalsArray*& OutSubObjectConditionals);
 
-	bool GetInternalChildSubObjectAndConditionalArrays(FInternalNetRefIndex InternalIndex, FInternalNetRefIndexArray*& OutChildSubObjectsArray, FSubObjectConditionalsArray*& OutSubObjectConditionals)
+	bool GetInternalChildSubObjectAndConditionalArrays(FInternalNetRefIndex InternalIndex, FInternalNetIndexArray*& OutChildSubObjectsArray, FSubObjectConditionalsArray*& OutSubObjectConditionals)
 	{
 		const FDependencyInfo* Entry = DependencyInfos.Find(InternalIndex);
 		const uint32 ArrayIndex = Entry ? Entry->ArrayIndices[ChildSubObjects] : FDependencyInfo::InvalidCacheIndex;
@@ -158,7 +162,7 @@ public:
 	}
 
 	template<EArrayType TypeIndex>
-	FInternalNetRefIndexArray* GetInternalIndexArray(FInternalNetRefIndex InternalIndex)
+	FInternalNetIndexArray* GetInternalIndexArray(FInternalNetRefIndex InternalIndex)
 	{
 		static_assert(TypeIndex != EArrayType::Count, "Invalid array type index");
 		const FDependencyInfo* Entry = DependencyInfos.Find(InternalIndex);
@@ -199,14 +203,14 @@ private:
 	};
 
 private:
-	FInternalNetRefIndexArray& GetOrCreateInternalIndexArray(FInternalNetRefIndex InternalIndex, EArrayType Type);
+	FInternalNetIndexArray& GetOrCreateInternalIndexArray(FInternalNetRefIndex InternalIndex, EArrayType Type);
 	FDependencyInfo& GetOrCreateCacheEntry(FInternalNetRefIndex InternalIndex);
 	
 	// Map to track the replicated objects with subObjects or dependencies
 	TMap<FInternalNetRefIndex, FDependencyInfo> DependencyInfos;
 
 	// Storage for DependentObjects and SubObjects
-	TSparseArray<FInternalNetRefIndexArray> DependentObjectsStorage;
+	TSparseArray<FInternalNetIndexArray> DependentObjectsStorage;
 
 	// Storage for SubObject conditionals
 	TSparseArray<FSubObjectConditionalsArray> SubObjectConditionalsStorage;
@@ -222,4 +226,4 @@ private:
 
 };
 
-}
+} // end namespace UE::Net::Private
diff --git a/Engine/Source/Runtime/Net/Iris/Private/Iris/ReplicationSystem/NetRefHandleManager.cpp b/Engine/Source/Runtime/Net/Iris/Private/Iris/ReplicationSystem/NetRefHandleManager.cpp
--- a/Engine/Source/Runtime/Net/Iris/Private/Iris/ReplicationSystem/NetRefHandleManager.cpp
+++ b/Engine/Source/Runtime/Net/Iris/Private/Iris/ReplicationSystem/NetRefHandleManager.cpp
@@ -78,13 +78,13 @@ const TCHAR* LexToString(EInternalDetachReason DetachReason)
 
 static bool IsValidIndex(FInternalNetRefIndex InternalIndex)
 {
-	return InternalIndex != FNetRefHandleManager::InvalidInternalIndex;
+	return InternalIndex != InvalidInternalNetRefIndex;
 }
 
 FNetRefHandleManager::FNetRefHandleManager(FReplicationProtocolManager& InReplicationProtocolManager)
 	: ReplicationProtocolManager(InReplicationProtocolManager)
 {
-	static_assert(InvalidInternalIndex == 0, "FNetRefHandleManager::InvalidInternalIndex has an unexpected value");
+	static_assert(InvalidInternalNetRefIndex == 0, "InvalidInternalNetRefIndex has an unexpected value");
 }
 
 void FNetRefHandleManager::Init(const FInitParams& InitParams)
@@ -104,10 +104,10 @@ void FNetRefHandleManager::Init(const FInitParams& InitParams)
 
 	ReplicationSystemId = InitParams.ReplicationSystemId;
 
-	// Must be a minimum of 1 to account for InvalidInternalIndex.
+	// Must be a minimum of 1 to account for InvalidInternalNetRefIndex.
 	uint32 PreAllocatedNetChunkedArrayCount = FMath::Clamp(InitParams.NetChunkedArrayCount, 1U, MaxActiveObjectCount);
 
-	// Calculate the highest internal index possible with the current preallocated buffers. Must be a minimum of 0 to support InvalidInternalIndex.
+	// Calculate the highest internal index possible with the current preallocated buffers. Must be a minimum of 0 to support InvalidInternalNetRefIndex.
 	HighestNetChunkedArrayInternalIndex = PreAllocatedNetChunkedArrayCount - 1;
 
 	// Initialize NetObjectList configs
@@ -123,8 +123,8 @@ void FNetRefHandleManager::Init(const FInitParams& InitParams)
 	ReplicatedObjectStateBuffers = TNetChunkedArray<uint8*>(PreAllocatedNetChunkedArrayCount, EInitMemory::Zero);
 	ReplicatedInstances = TNetChunkedArray<TObjectPtr<UObject>>(PreAllocatedNetChunkedArrayCount, EInitMemory::Zero);
 
-	// For convenience we initialize ReplicatedObjectData for InvalidInternalIndex so that GetReplicatedObjectDataNoCheck returns something useful.
-	ReplicatedObjectData[InvalidInternalIndex].Reset();
+	// For convenience we initialize ReplicatedObjectData for InvalidInternalNetRefIndex so that GetReplicatedObjectDataNoCheck returns something useful.
+	ReplicatedObjectData[InvalidInternalNetRefIndex].Reset();
 
 	// Init all NetBitArrays here
 	{
@@ -181,7 +181,7 @@ FInternalNetRefIndex FNetRefHandleManager::GrowNetObjectLists()
 	// We already are at the max, return InvalidIndex and abort.
 	if (CurrentMaxInternalNetRefIndex >= MaxActiveObjectCount)
 	{
-		return InvalidInternalIndex;
+		return InvalidInternalNetRefIndex;
 	}
 
     CurrentMaxInternalNetRefIndex += InternalNetRefIndexGrowSize;
@@ -257,7 +257,7 @@ uint64 FNetRefHandleManager::GetNextNetRefHandleSerial(uint64 HandleId) const
 FInternalNetRefIndex FNetRefHandleManager::GetNextFreeInternalIndex() const
 {
 	const uint32 NextFreeIndex = AssignedInternalIndices.FindFirstZero();
-	return NextFreeIndex != FNetBitArray::InvalidIndex ? NextFreeIndex : InvalidInternalIndex;
+	return NextFreeIndex != FNetBitArray::InvalidIndex ? NextFreeIndex : InvalidInternalNetRefIndex;
 }
 
 FInternalNetRefIndex FNetRefHandleManager::InternalCreateNetObject(const FNetRefHandle NetRefHandle, const FNetHandle GlobalHandle, const FCreateNetObjectParams& Params)
@@ -266,7 +266,7 @@ FInternalNetRefIndex FNetRefHandleManager::InternalCreateNetObject(const FNetRef
 	{
 		UE_LOG(LogIris, Error, TEXT("NetRefHandleManager: Maximum active object count reached (%d/%d)."), ActiveObjectCount, MaxActiveObjectCount);
 		ensureMsgf(false, TEXT("NetRefHandleManager: Maximum active object count reached (%d/%d)."), ActiveObjectCount, MaxActiveObjectCount);
-		return InvalidInternalIndex;
+		return InvalidInternalNetRefIndex;
 	}
 
 	// Verify that the handle is free
@@ -274,7 +274,7 @@ FInternalNetRefIndex FNetRefHandleManager::InternalCreateNetObject(const FNetRef
 	{
 		UE_LOG(LogIris, Error, TEXT("NetRefHandleManager::InternalCreateNetObject %s already exists"), *NetRefHandle.ToString());
 		ensureMsgf(false, TEXT("NetRefHandleManager::InternalCreateNetObject %s already exists"), *NetRefHandle.ToString());
-		return InvalidInternalIndex;
+		return InvalidInternalNetRefIndex;
 	}
 
 	// Verify that we are not trying to StartRepliation for an object using the same NetRefHandle as one that is pending destroy.
@@ -288,7 +288,7 @@ FInternalNetRefIndex FNetRefHandleManager::InternalCreateNetObject(const FNetRef
 				{
 					UE_LOG(LogIris, Error, TEXT("NetRefHandleManager::InternalCreateNetObject Cannot create NetObject %s for object already replicated as %s marked for destroy. Only supported if using AsyncStopReplication."), *NetRefHandle.ToString(), *PrintObjectFromIndex(PendingDestroyInfo.InternalIndex));
 					ensureMsgf(false, TEXT("NetRefHandleManager::InternalCreateNetObject Cannot create NetObject %s for object marked for destroy. Only supported if using AsyncStopReplication."), *NetRefHandle.ToString());
-					return InvalidInternalIndex;
+					return InvalidInternalNetRefIndex;
 				}
 				else
 				{
@@ -302,15 +302,15 @@ FInternalNetRefIndex FNetRefHandleManager::InternalCreateNetObject(const FNetRef
 	uint32 InternalIndex = GetNextFreeInternalIndex();
 
 	// Try to grow the NetObjectLists if no more indexes are available.
-	if (InternalIndex == InvalidInternalIndex)
+	if (InternalIndex == InvalidInternalNetRefIndex)
 	{
 		InternalIndex = GrowNetObjectLists();
 
 		// If we could not grow anymore, kill the process now. The system cannot replicate objects anymore and the game behavior is undefined.
-		if (InternalIndex == InvalidInternalIndex)
+		if (InternalIndex == InvalidInternalNetRefIndex)
 		{
 			UE_LOG(LogIris, Fatal, TEXT("NetRefHandleManager: Hit the maximum limit of active replicated objects: %u. Aborting since we cannot replicate %s"), MaxActiveObjectCount, Params.ReplicationProtocol->DebugName->Name);
-			return InvalidInternalIndex;
+			return InvalidInternalNetRefIndex;
 		}
 	}
 
@@ -368,7 +368,7 @@ FInternalNetRefIndex FNetRefHandleManager::InternalCreateNetObject(const FNetRef
 
 void FNetRefHandleManager::AttachInstanceProtocol(FInternalNetRefIndex InternalIndex, const FReplicationInstanceProtocol* InstanceProtocol, UObject* Instance)
 {
-	if (ensure((InternalIndex != InvalidInternalIndex) && InstanceProtocol))
+	if (ensure((InternalIndex != InvalidInternalNetRefIndex) && InstanceProtocol))
 	{
 		FReplicatedObjectData& Data = ReplicatedObjectData[InternalIndex];
 		Data.InstanceProtocol = InstanceProtocol;
@@ -383,7 +383,7 @@ void FNetRefHandleManager::AttachInstanceProtocol(FInternalNetRefIndex InternalI
 
 const FReplicationInstanceProtocol* FNetRefHandleManager::DetachInstanceProtocol(FInternalNetRefIndex InternalIndex)
 {
-	if (ensure(InternalIndex != InvalidInternalIndex))
+	if (ensure(InternalIndex != InvalidInternalNetRefIndex))
 	{
 		FReplicatedObjectData& Data = ReplicatedObjectData[InternalIndex];
 		const FReplicationInstanceProtocol* InstanceProtocol = Data.InstanceProtocol;
@@ -402,7 +402,7 @@ const FReplicationInstanceProtocol* FNetRefHandleManager::DetachInstanceProtocol
 
 bool FNetRefHandleManager::HasInstanceProtocol(FInternalNetRefIndex InternalIndex) const
 {
-	check(InternalIndex == InvalidInternalIndex || AssignedInternalIndices.GetBit(InternalIndex));
+	check(InternalIndex == InvalidInternalNetRefIndex || AssignedInternalIndices.GetBit(InternalIndex));
 	return ReplicatedObjectData[InternalIndex].InstanceProtocol != nullptr;
 }
 
@@ -440,7 +440,7 @@ FNetRefHandle FNetRefHandleManager::CreateNetObject(FNetRefHandle WantedHandle,
 	FNetRefHandle NetRefHandle = WantedHandle;
 
 	const FInternalNetRefIndex InternalIndex = InternalCreateNetObject(NetRefHandle, GlobalHandle, Params);
-	if (InternalIndex != InvalidInternalIndex)
+	if (InternalIndex != InvalidInternalNetRefIndex)
 	{
 		FReplicatedObjectData& Data = ReplicatedObjectData[InternalIndex];
 
@@ -481,7 +481,7 @@ FNetRefHandle FNetRefHandleManager::CreateNetObjectFromRemote(FNetRefHandle Want
 	FNetRefHandle NetRefHandle = MakeNetRefHandle(WantedHandle.GetId(), ReplicationSystemId);
 
 	const FInternalNetRefIndex InternalIndex = InternalCreateNetObject(NetRefHandle, FNetHandle(), Params);
-	if (InternalIndex != InvalidInternalIndex)
+	if (InternalIndex != InvalidInternalNetRefIndex)
 	{
 		FReplicatedObjectData& Data = ReplicatedObjectData[InternalIndex];
 
@@ -545,7 +545,7 @@ void FNetRefHandleManager::InternalDestroyNetObject(FInternalNetRefIndex Interna
 	}
 	
 	// If this is a RootObject, remove all subobjects from the list
-	if (FNetDependencyData::FInternalNetRefIndexArray* SubObjectArray = SubObjects.GetInternalIndexArray<FNetDependencyData::EArrayType::SubObjects>(InternalIndex))
+	if (FNetDependencyData::FInternalNetIndexArray* SubObjectArray = SubObjects.GetInternalIndexArray<FNetDependencyData::EArrayType::SubObjects>(InternalIndex))
 	{
 		for (FInternalNetRefIndex SubObjectInternalIndex : *SubObjectArray)
 		{
@@ -554,13 +554,13 @@ void FNetRefHandleManager::InternalDestroyNetObject(FInternalNetRefIndex Interna
 		SubObjectArray->Reset(0U);
 	}
 	// Clear ChildSubObjectArray
-	if (FNetDependencyData::FInternalNetRefIndexArray* ChildSubObjectArray = SubObjects.GetInternalIndexArray<FNetDependencyData::EArrayType::ChildSubObjects>(InternalIndex))
+	if (FNetDependencyData::FInternalNetIndexArray* ChildSubObjectArray = SubObjects.GetInternalIndexArray<FNetDependencyData::EArrayType::ChildSubObjects>(InternalIndex))
 	{
 		ChildSubObjectArray->Reset(0U);
 	}
 
 	// If we are a subobject remove from owner and hierarchical parents
-	if (Data.SubObjectRootIndex != InvalidInternalIndex)
+	if (Data.SubObjectRootIndex != InvalidInternalNetRefIndex)
 	{
 		InternalRemoveSubObject(Data.SubObjectRootIndex, InternalIndex);
 	}
@@ -765,8 +765,8 @@ bool FNetRefHandleManager::AddSubObject(FNetRefHandle ParentObjectHandle, FNetRe
 	const FInternalNetRefIndex ParentObjectInternalIndex = GetInternalIndex(ParentObjectHandle);
 	const FInternalNetRefIndex SubObjectInternalIndex = GetInternalIndex(SubObjectHandle);
 
-	const bool bIsValidParent = ensure(ParentObjectInternalIndex != InvalidInternalIndex);
-	const bool bIsValidSubObject = ensure(SubObjectInternalIndex != InvalidInternalIndex);
+	const bool bIsValidParent = ensure(ParentObjectInternalIndex != InvalidInternalNetRefIndex);
+	const bool bIsValidSubObject = ensure(SubObjectInternalIndex != InvalidInternalNetRefIndex);
 
 	if (!bIsValidParent || !bIsValidSubObject)
 	{
@@ -780,7 +780,7 @@ bool FNetRefHandleManager::AddSubObject(FNetRefHandle ParentObjectHandle, FNetRe
 	UE_NET_TRACE_SUBOBJECT(RootObjectHandle, SubObjectHandle);
 #endif
 
-	const FInternalNetRefIndex RelativeOtherSubObjectInternalIndex = EnumHasAnyFlags(Flags, EAddSubObjectFlags::InsertAtStartOrAfterOther) ? GetInternalIndex(RelativeOtherSubObjectHandle) : InvalidInternalIndex;
+	const FInternalNetRefIndex RelativeOtherSubObjectInternalIndex = EnumHasAnyFlags(Flags, EAddSubObjectFlags::InsertAtStartOrAfterOther) ? GetInternalIndex(RelativeOtherSubObjectHandle) : InvalidInternalNetRefIndex;
 	return InternalAddSubObject(RootObjectInternalIndex, ParentObjectInternalIndex, SubObjectInternalIndex, RelativeOtherSubObjectInternalIndex, Flags);
 }
 
@@ -789,13 +789,13 @@ bool FNetRefHandleManager::InternalAddSubObject(FInternalNetRefIndex RootObjectI
 	using namespace UE::Net::Private;
 
 	FReplicatedObjectData& SubObjectData = GetReplicatedObjectDataNoCheck(SubObjectInternalIndex);
-	if (!ensureMsgf(SubObjectData.SubObjectRootIndex == InvalidInternalIndex, TEXT("FNetRefHandleManager::AddSubObject %s is already marked as a subobject"), ToCStr(SubObjectData.RefHandle.ToString())))
+	if (!ensureMsgf(SubObjectData.SubObjectRootIndex == InvalidInternalNetRefIndex, TEXT("FNetRefHandleManager::AddSubObject %s is already marked as a subobject"), ToCStr(SubObjectData.RefHandle.ToString())))
 	{
 		return false;
 	}
 
 	// Add the subobject to the RootObject, this array contains all subobjects associated with the root in the order they are added
-	FNetDependencyData::FInternalNetRefIndexArray& SubObjectArray = SubObjects.GetOrCreateInternalIndexArray<FNetDependencyData::SubObjects>(RootObjectInternalIndex);
+	FNetDependencyData::FInternalNetIndexArray& SubObjectArray = SubObjects.GetOrCreateInternalIndexArray<FNetDependencyData::SubObjects>(RootObjectInternalIndex);
 	SubObjectArray.Add(SubObjectInternalIndex);
 	
 	SubObjectData.SubObjectRootIndex = RootObjectInternalIndex;
@@ -805,7 +805,7 @@ bool FNetRefHandleManager::InternalAddSubObject(FInternalNetRefIndex RootObjectI
 
 	FInternalNetRefIndex ParentOfSubObjectIndex = RootObjectInternalIndex;
 	
-	if (OptionalParentSubObjectInternalIndex != InvalidInternalIndex && OptionalParentSubObjectInternalIndex != RootObjectInternalIndex)
+	if (OptionalParentSubObjectInternalIndex != InvalidInternalNetRefIndex && OptionalParentSubObjectInternalIndex != RootObjectInternalIndex)
 	{
 		// Make sure the non-rootobject Parent is a subobject of the same RootObject
 		const bool bIsValidOuter = SubObjectArray.Contains(OptionalParentSubObjectInternalIndex);
@@ -817,7 +817,7 @@ bool FNetRefHandleManager::InternalAddSubObject(FInternalNetRefIndex RootObjectI
 
 	// Add the subobject to its Parent's list which is what we use to define replication order.
 	FNetDependencyData::FSubObjectConditionalsArray* SubObjectConditionalsArray = nullptr;
-	FNetDependencyData::FInternalNetRefIndexArray& ChildSubObjectArray = SubObjects.GetOrCreateInternalChildSubObjectsArray(ParentOfSubObjectIndex, SubObjectConditionalsArray);
+	FNetDependencyData::FInternalNetIndexArray& ChildSubObjectArray = SubObjects.GetOrCreateInternalChildSubObjectsArray(ParentOfSubObjectIndex, SubObjectConditionalsArray);
 
 	if (EnumHasAnyFlags(Flags, EAddSubObjectFlags::InsertAtStart))
 	{
@@ -831,7 +831,7 @@ bool FNetRefHandleManager::InternalAddSubObject(FInternalNetRefIndex RootObjectI
 	else if (EnumHasAnyFlags(Flags, EAddSubObjectFlags::InsertAtStartOrAfterOther))
 	{
 		// Insert the object right after the other one, or at the start if it wasn't in the list
-		const int32 RelativeIndex = RelativeOtherSubObjectInternalIndex != FNetRefHandleManager::InvalidInternalIndex ? ChildSubObjectArray.Find(RelativeOtherSubObjectInternalIndex) : INDEX_NONE;
+		const int32 RelativeIndex = RelativeOtherSubObjectInternalIndex != InvalidInternalNetRefIndex ? ChildSubObjectArray.Find(RelativeOtherSubObjectInternalIndex) : INDEX_NONE;
 		const int32 InsertIndex = (RelativeIndex == INDEX_NONE) ? 0 : (RelativeIndex + 1);
 
 		ChildSubObjectArray.Insert(SubObjectInternalIndex, InsertIndex);
@@ -865,7 +865,7 @@ bool FNetRefHandleManager::InternalAddSubObject(FInternalNetRefIndex RootObjectI
 void FNetRefHandleManager::InternalRemoveSubObject(FInternalNetRefIndex RootObjectInternalIndex, FInternalNetRefIndex SubObjectInternalIndex, bool bRemoveFromSubObjectArray)
 {
 	// both must be valid
-	if (RootObjectInternalIndex == InvalidInternalIndex || SubObjectInternalIndex == InvalidInternalIndex)
+	if (RootObjectInternalIndex == InvalidInternalNetRefIndex || SubObjectInternalIndex == InvalidInternalNetRefIndex)
 	{
 		return;
 	}
@@ -876,15 +876,15 @@ void FNetRefHandleManager::InternalRemoveSubObject(FInternalNetRefIndex RootObje
 	if (bRemoveFromSubObjectArray)
 	{
 		// Remove subobject from the root object's list
-		if (FNetDependencyData::FInternalNetRefIndexArray* SubObjectArray = SubObjects.GetInternalIndexArray<FNetDependencyData::EArrayType::SubObjects>(RootObjectInternalIndex))
+		if (FNetDependencyData::FInternalNetIndexArray* SubObjectArray = SubObjects.GetInternalIndexArray<FNetDependencyData::EArrayType::SubObjects>(RootObjectInternalIndex))
 		{
 			SubObjectArray->Remove(SubObjectInternalIndex);
 		}
 
 		// Remove subobject from the parent's list
-		if (SubObjectData.SubObjectParentIndex != InvalidInternalIndex)
+		if (SubObjectData.SubObjectParentIndex != InvalidInternalNetRefIndex)
 		{
-			FNetDependencyData::FInternalNetRefIndexArray* ChildSubObjectArray;
+			FNetDependencyData::FInternalNetIndexArray* ChildSubObjectArray;
 			FNetDependencyData::FSubObjectConditionalsArray* SubObjectConditionsArray;
 
 			if (SubObjects.GetInternalChildSubObjectAndConditionalArrays(SubObjectData.SubObjectParentIndex, ChildSubObjectArray, SubObjectConditionsArray))
@@ -909,8 +909,8 @@ void FNetRefHandleManager::InternalRemoveSubObject(FInternalNetRefIndex RootObje
 		ObjectsWithDirtyCreationDependencies.SetBit(SubObjectInternalIndex);
 	}
 
-	SubObjectData.SubObjectRootIndex = InvalidInternalIndex;
-	SubObjectData.SubObjectParentIndex = InvalidInternalIndex;
+	SubObjectData.SubObjectRootIndex = InvalidInternalNetRefIndex;
+	SubObjectData.SubObjectParentIndex = InvalidInternalNetRefIndex;
 
 	SetIsSubObject(SubObjectInternalIndex, false);
 }
@@ -920,10 +920,10 @@ void FNetRefHandleManager::RemoveSubObject(FNetRefHandle Handle)
 	FInternalNetRefIndex SubObjectInternalIndex = GetInternalIndex(Handle);
 	checkSlow(SubObjectInternalIndex);
 
-	if (SubObjectInternalIndex != InvalidInternalIndex)
+	if (SubObjectInternalIndex != InvalidInternalNetRefIndex)
 	{
 		const FInternalNetRefIndex RootObjectInternalIndex = ReplicatedObjectData[SubObjectInternalIndex].SubObjectRootIndex;
-		if (RootObjectInternalIndex != InvalidInternalIndex)
+		if (RootObjectInternalIndex != InvalidInternalNetRefIndex)
 		{
 			InternalRemoveSubObject(RootObjectInternalIndex, SubObjectInternalIndex);
 		}
@@ -932,12 +932,12 @@ void FNetRefHandleManager::RemoveSubObject(FNetRefHandle Handle)
 
 bool FNetRefHandleManager::SetSubObjectNetCondition(FInternalNetRefIndex SubObjectInternalIndex, FLifeTimeConditionStorage SubObjectCondition, bool& bOutWasModified)
 {
-	if (ensure(SubObjectInternalIndex != InvalidInternalIndex))
+	if (ensure(SubObjectInternalIndex != InvalidInternalNetRefIndex))
 	{
 		const FInternalNetRefIndex SubObjectParentIndex = ReplicatedObjectData[SubObjectInternalIndex].SubObjectParentIndex;
-		if (ensure(SubObjectParentIndex != InvalidInternalIndex))
+		if (ensure(SubObjectParentIndex != InvalidInternalNetRefIndex))
 		{
-			FNetDependencyData::FInternalNetRefIndexArray* SubObjectsArray;
+			FNetDependencyData::FInternalNetIndexArray* SubObjectsArray;
 			FNetDependencyData::FSubObjectConditionalsArray* SubObjectConditionals;
 			
 			if (SubObjects.GetInternalChildSubObjectAndConditionalArrays(SubObjectParentIndex, SubObjectsArray, SubObjectConditionals))
@@ -993,7 +993,7 @@ FNetRefHandle FNetRefHandleManager::GetRootObjectOfAnyObject(FNetRefHandle NetRe
 FNetRefHandle FNetRefHandleManager::GetRootObjectOfSubObject(FNetRefHandle SubObjectRefHandle) const
 {
 	const FInternalNetRefIndex SubObjectInternalIndex = GetInternalIndex(SubObjectRefHandle);
-	const FInternalNetRefIndex OwnerInternalIndex = IsValidIndex(SubObjectInternalIndex) ? ReplicatedObjectData[SubObjectInternalIndex].SubObjectRootIndex : InvalidInternalIndex;
+	const FInternalNetRefIndex OwnerInternalIndex = IsValidIndex(SubObjectInternalIndex) ? ReplicatedObjectData[SubObjectInternalIndex].SubObjectRootIndex : InvalidInternalNetRefIndex;
 
 	return IsValidIndex(OwnerInternalIndex) ? ReplicatedObjectData[OwnerInternalIndex].RefHandle : FNetRefHandle();
 }
@@ -1007,7 +1007,7 @@ FInternalNetRefIndex FNetRefHandleManager::GetRootObjectInternalIndexOfAnyObject
 
 FInternalNetRefIndex FNetRefHandleManager::GetRootObjectInternalIndexOfSubObject(FInternalNetRefIndex SubObjectIndex) const
 {
-	// If InvalidInternalIndex is passed we can rely on that data being default initialized so it will return InvalidInternalIndex too.
+	// If InvalidInternalNetRefIndex is passed we can rely on that data being default initialized so it will return InvalidInternalNetRefIndex too.
 	return ReplicatedObjectData[SubObjectIndex].SubObjectRootIndex;
 }
 
@@ -1024,8 +1024,8 @@ bool FNetRefHandleManager::AddDependentObject(FNetRefHandle ParentRefHandle, FNe
 	FInternalNetRefIndex ParentInternalIndex = GetInternalIndex(ParentRefHandle);
 	FInternalNetRefIndex DependentObjectInternalIndex = GetInternalIndex(DependentObjectRefHandle);
 
-	const bool bIsValidOwner = ensure(ParentInternalIndex != InvalidInternalIndex);
-	const bool bIsValidDependentObject = ensure(DependentObjectInternalIndex != InvalidInternalIndex);
+	const bool bIsValidOwner = ensure(ParentInternalIndex != InvalidInternalNetRefIndex);
+	const bool bIsValidDependentObject = ensure(DependentObjectInternalIndex != InvalidInternalNetRefIndex);
 
 	if (!(bIsValidOwner && bIsValidDependentObject))
 	{
@@ -1055,7 +1055,7 @@ bool FNetRefHandleManager::AddDependentObject(FNetRefHandle ParentRefHandle, FNe
 	}
 
 	FNetDependencyData::FDependentObjectInfoArray& ParentDependentObjectsArray = SubObjects.GetOrCreateDependentObjectInfoArray(ParentInternalIndex);
-	FNetDependencyData::FInternalNetRefIndexArray& DependentParentObjectArray = SubObjects.GetOrCreateInternalIndexArray<FNetDependencyData::DependentParentObjects>(DependentObjectInternalIndex);
+	FNetDependencyData::FInternalNetIndexArray& DependentParentObjectArray = SubObjects.GetOrCreateInternalIndexArray<FNetDependencyData::DependentParentObjects>(DependentObjectInternalIndex);
 	
 	// Make sure parent didn't set the child as a dependent already
 	{
@@ -1107,7 +1107,7 @@ void FNetRefHandleManager::RemoveDependentObject(FNetRefHandle DependentHandle)
 {
 	FInternalNetRefIndex DependentInternalIndex = GetInternalIndex(DependentHandle);
 
-	if (DependentInternalIndex != InvalidInternalIndex)
+	if (DependentInternalIndex != InvalidInternalNetRefIndex)
 	{
 		InternalRemoveAllDependencies(DependentInternalIndex);
 	}
@@ -1137,7 +1137,7 @@ void FNetRefHandleManager::RemoveCreationDependency(FNetRefHandle Parent, FNetRe
 
 bool FNetRefHandleManager::InternalAddCreationDependency(FInternalNetRefIndex ParentIndex, FInternalNetRefIndex ChildIndex)
 {
-	if (ParentIndex == InvalidInternalIndex || ChildIndex == InvalidInternalIndex)
+	if (ParentIndex == InvalidInternalNetRefIndex || ChildIndex == InvalidInternalNetRefIndex)
 	{
 		ensureMsgf(false, TEXT("AddCreationDependency failed due to invalid reference. Parent: %s | Child: %s"), *PrintObjectFromIndex(ParentIndex), *PrintObjectFromIndex(ChildIndex));
 		return false;
@@ -1176,7 +1176,7 @@ bool FNetRefHandleManager::InternalAddCreationDependency(FInternalNetRefIndex Pa
 
 void FNetRefHandleManager::InternalRemoveCreationDependency(FInternalNetRefIndex ParentIndex, FInternalNetRefIndex ChildIndex)
 {
-	if (ParentIndex == InvalidInternalIndex || ChildIndex == InvalidInternalIndex)
+	if (ParentIndex == InvalidInternalNetRefIndex || ChildIndex == InvalidInternalNetRefIndex)
 	{
 		ensureMsgf(false, TEXT("RemoveCreationDependency failed due to invalid reference. Parent: %s | Child: %s"), *PrintObjectFromIndex(ParentIndex), *PrintObjectFromIndex(ChildIndex));
 		return;
@@ -1248,7 +1248,7 @@ void FNetRefHandleManager::InternalRemoveDependentObject(FInternalNetRefIndex Pa
 {
 	if (EnumHasAnyFlags(Flags, ERemoveDependentObjectFlags::RemoveFromDependentParentObjects))
 	{
-		if (FNetDependencyData::FInternalNetRefIndexArray* ParentObjectArray = SubObjects.GetInternalIndexArray<FNetDependencyData::DependentParentObjects>(DependentInternalIndex))
+		if (FNetDependencyData::FInternalNetIndexArray* ParentObjectArray = SubObjects.GetInternalIndexArray<FNetDependencyData::DependentParentObjects>(DependentInternalIndex))
 		{
 			ParentObjectArray->Remove(ParentInternalIndex);
 			if (ParentObjectArray->Num() == 0)
@@ -1284,7 +1284,7 @@ void FNetRefHandleManager::InternalRemoveDependentObject(FInternalNetRefIndex Pa
 void FNetRefHandleManager::InternalRemoveAllDependencies(FInternalNetRefIndex DependentInternalIndex)
 {
 	// Remove from all parents
-	if (FNetDependencyData::FInternalNetRefIndexArray* ParentObjectArray = SubObjects.GetInternalIndexArray<FNetDependencyData::DependentParentObjects>(DependentInternalIndex))
+	if (FNetDependencyData::FInternalNetIndexArray* ParentObjectArray = SubObjects.GetInternalIndexArray<FNetDependencyData::DependentParentObjects>(DependentInternalIndex))
 	{
 		for (FInternalNetRefIndex ParentInternalIndex : *ParentObjectArray)
 		{
@@ -1346,7 +1346,7 @@ void FNetRefHandleManager::RemoveDependentObject(FNetRefHandle ParentHandle, FNe
 	FInternalNetRefIndex ParentInternalIndex = GetInternalIndex(ParentHandle);
 	FInternalNetRefIndex DependentInternalIndex = GetInternalIndex(DependentHandle);
 
-	if ((ParentInternalIndex == InvalidInternalIndex) || (DependentInternalIndex == InvalidInternalIndex))
+	if ((ParentInternalIndex == InvalidInternalNetRefIndex) || (DependentInternalIndex == InvalidInternalNetRefIndex))
 	{
 		return;
 	}
@@ -1387,7 +1387,7 @@ bool FNetRefHandleManager::AddingDependencyWouldCreateCircularDependency(FIntern
 void FNetRefHandleManager::SetIsDebugObject(FNetRefHandle Handle)
 {
 	FInternalNetRefIndex ObjectInternalIndex = GetInternalIndex(Handle);
-	if (ObjectInternalIndex != InvalidInternalIndex && ensureMsgf(ReplicatedObjectRefCount[ObjectInternalIndex] == 0U, TEXT("SetIsDebugObject can only be called before we have started to replicate an object. %s"), *PrintObjectFromIndex(ObjectInternalIndex)))
+	if (ObjectInternalIndex != InvalidInternalNetRefIndex && ensureMsgf(ReplicatedObjectRefCount[ObjectInternalIndex] == 0U, TEXT("SetIsDebugObject can only be called before we have started to replicate an object. %s"), *PrintObjectFromIndex(ObjectInternalIndex)))
 	{
 		ReplicatedObjectData[ObjectInternalIndex].bIsDebugObject = 1U;
 	}
@@ -1395,7 +1395,7 @@ void FNetRefHandleManager::SetIsDebugObject(FNetRefHandle Handle)
 
 void FNetRefHandleManager::SetShouldPropagateChangedStates(FInternalNetRefIndex ObjectInternalIndex, bool bShouldPropagateChangedStates)
 {
-	if (ObjectInternalIndex != InvalidInternalIndex)
+	if (ObjectInternalIndex != InvalidInternalNetRefIndex)
 	{
 		if (bShouldPropagateChangedStates)
 		{
@@ -1427,7 +1427,7 @@ void FNetRefHandleManager::AddReferencedObjects(FReferenceCollector& Collector)
 
 void FNetRefHandleManager::StoreCachedCreationInfo(FInternalNetRefIndex InternalIndex, TUniquePtr<FNetObjectCreationHeader>&& Header)
 {
-	if (InternalIndex != InvalidInternalIndex)
+	if (InternalIndex != InvalidInternalNetRefIndex)
 	{
 		FReplicatedObjectData& ObjectData = GetReplicatedObjectDataNoCheck(InternalIndex);
 		ObjectData.bHasCachedCreationInfo = 1U;
@@ -1437,7 +1437,7 @@ void FNetRefHandleManager::StoreCachedCreationInfo(FInternalNetRefIndex Internal
 
 void FNetRefHandleManager::ClearCachedCreationInfo(FInternalNetRefIndex InternalIndex)
 {
-	if (InternalIndex != InvalidInternalIndex)
+	if (InternalIndex != InvalidInternalNetRefIndex)
 	{
 		FReplicatedObjectData& ObjectData = GetReplicatedObjectDataNoCheck(InternalIndex);
 		ObjectData.bHasCachedCreationInfo = 0U;
@@ -1535,7 +1535,7 @@ void FNetRefHandleManager::OnPostSendUpdate()
 
 FString FNetRefHandleManager::PrintObjectFromIndex(FInternalNetRefIndex ObjectIndex) const
 {
-	if (ObjectIndex != InvalidInternalIndex)
+	if (ObjectIndex != InvalidInternalNetRefIndex)
 	{
 		const bool bIsIndexAssigned = AssignedInternalIndices.GetBit(ObjectIndex);
 		if (!bIsIndexAssigned)
@@ -1547,7 +1547,7 @@ FString FNetRefHandleManager::PrintObjectFromIndex(FInternalNetRefIndex ObjectIn
 		const FReplicatedObjectData& ObjectData = GetReplicatedObjectDataNoCheck(ObjectIndex);
 		const FNetRefHandle NetRefHandle = ObjectData.RefHandle;
 
-		if (ObjectData.SubObjectRootIndex == InvalidInternalIndex)
+		if (ObjectData.SubObjectRootIndex == InvalidInternalNetRefIndex)
 		{
 			return FString::Printf(TEXT("RootObject %s (InternalIndex: %u) (%s) Protocol: (%s)"), *GetNameSafe(ReplicatedInstances[ObjectIndex]), ObjectIndex, *NetRefHandle.ToString(), ObjectData.Protocol->DebugName->Name);
 		}
@@ -1568,7 +1568,7 @@ FString FNetRefHandleManager::PrintObjectFromIndex(FInternalNetRefIndex ObjectIn
 FString FNetRefHandleManager::PrintObjectFromNetRefHandle(FNetRefHandle ObjectHandle) const
 { 
 	const FInternalNetRefIndex ObjectIndex = GetInternalIndex(ObjectHandle);
-	if (ObjectIndex != InvalidInternalIndex)
+	if (ObjectIndex != InvalidInternalNetRefIndex)
 	{
 		return PrintObjectFromIndex(ObjectIndex);
 	}
diff --git a/Engine/Source/Runtime/Net/Iris/Private/Iris/ReplicationSystem/NetRefHandleManager.h b/Engine/Source/Runtime/Net/Iris/Private/Iris/ReplicationSystem/NetRefHandleManager.h
--- a/Engine/Source/Runtime/Net/Iris/Private/Iris/ReplicationSystem/NetRefHandleManager.h
+++ b/Engine/Source/Runtime/Net/Iris/Private/Iris/ReplicationSystem/NetRefHandleManager.h
@@ -15,6 +15,8 @@
 #include "Iris/IrisConstants.h"
 
 #include "Iris/ReplicationState/ReplicationStateDescriptor.h"
+
+#include "Iris/ReplicationSystem/InternalNetRefIndexManager.h"
 #include "Iris/ReplicationSystem/NetDependencyData.h"
 #include "Iris/ReplicationSystem/NetObjectFactory.h"
 #include "Iris/ReplicationSystem/NetObjectFactoryRegistry.h"
@@ -45,8 +47,6 @@ namespace UE::Net
 namespace UE::Net::Private
 {
 
-typedef uint32 FInternalNetRefIndex;
-
 enum class EAddSubObjectFlags : uint32
 {
 	None = 0U,
@@ -88,8 +88,6 @@ consteval uint32 GetCloseReasonBits()
 class FNetRefHandleManager
 {
 public:
-	// We will never assign InvalidInternalIndex
-	enum { InvalidInternalIndex = 0U };
 
 	// We need to store some internal data for Replicated objects
 	// $TODO: This should be split up into separate array according to usage patterns, it is getting a bit too large
@@ -116,10 +114,10 @@ public:
 		uint8* ReceiveStateBuffer = nullptr;
 		
 		/** Subobjects only: Internal index of the RootObject of this subobject */
-		FInternalNetRefIndex SubObjectRootIndex = InvalidInternalIndex;
+		FInternalNetRefIndex SubObjectRootIndex = InvalidInternalNetRefIndex;
 		
 		/** Subobjects only: Internal index of the ParentObject of this subobject */
-		FInternalNetRefIndex SubObjectParentIndex = InvalidInternalIndex;
+		FInternalNetRefIndex SubObjectParentIndex = InvalidInternalNetRefIndex;
 
 		/** SubObjects only: Determine if child subobjects should replicate before or after parent */
 		uint16 bReplicateChildSubObjectsAfterParent : 1U = 0;
@@ -151,10 +149,10 @@ public:
 		EInternalDetachReason InternalDetachReason = EInternalDetachReason::Normal;
 
 		/** Returns true if this is a root object */
-		bool IsRootObject() const { return SubObjectRootIndex == InvalidInternalIndex; }
+		bool IsRootObject() const { return SubObjectRootIndex == InvalidInternalNetRefIndex; }
 	
 		/** Returns true if this is a SubObject */
-		bool IsSubObject() const { return SubObjectRootIndex != InvalidInternalIndex; }
+		bool IsSubObject() const { return SubObjectRootIndex != InvalidInternalNetRefIndex; }
 
 		/** Returns true if this is a child dependent object  */
 		bool IsDependentObject() const { return bIsDependentObject; }
@@ -319,7 +317,7 @@ public:
 	// If InternalIndex is a root object or invalid it will be returned. If it's a subobject then its root object index will be returned.
 	FInternalNetRefIndex GetRootObjectInternalIndexOfAnyObject(FInternalNetRefIndex InternalIndex) const;
 
-	// If SubObjectIndex is valid and a subobject then its root object index is returned, otherwise InvalidInternalIndex is returned.
+	// If SubObjectIndex is valid and a subobject then its root object index is returned, otherwise InvalidInternalNetRefIndex is returned.
 	FInternalNetRefIndex GetRootObjectInternalIndexOfSubObject(FInternalNetRefIndex SubObjectIndex) const;
 
 	bool SetSubObjectNetCondition(FInternalNetRefIndex SubObjectInternalIndex, FLifeTimeConditionStorage SubObjectCondition, bool& bOutWasModified);
@@ -502,7 +500,7 @@ private:
 
 	uint64 GetNextNetRefHandleSerial(uint64 HandleIndex) const;
 
-	// Get the next free internal index, returns InvalidInternalIndex if a free one cannot be found
+	// Get the next free internal index, returns InvalidInternalNetRefIndex if a free one cannot be found
 	FInternalNetRefIndex GetNextFreeInternalIndex() const;
 
 	bool InternalAddSubObject(FInternalNetRefIndex RootObjectInternalIndex, FInternalNetRefIndex ParentSubObjectInternalIndex, FInternalNetRefIndex SubObjectInternalIndex, FInternalNetRefIndex RelativeOtherSubObjectInternalIndex, EAddSubObjectFlags Flags);
@@ -516,7 +514,7 @@ private:
 
 	/**
 	 * Grow the NetObjectLists to fit more internal indexes.
-	 * Returns the newest available index if the increase worked or InvalidInternalIndex if we hit the maximum.
+	 * Returns the newest available index if the increase worked or InvalidInternalNetRefIndex if we hit the maximum.
 	 */
 	FInternalNetRefIndex GrowNetObjectLists();
 
@@ -677,13 +675,13 @@ const FNetRefHandleManager::FReplicatedObjectData& FNetRefHandleManager::GetRepl
 FInternalNetRefIndex FNetRefHandleManager::GetInternalIndex(FNetRefHandle Handle) const
 {
 	const FInternalNetRefIndex* InternalIndex = RefHandleToInternalIndex.Find(Handle);
-	return InternalIndex ? *InternalIndex : FNetRefHandleManager::InvalidInternalIndex;
+	return InternalIndex ? *InternalIndex : InvalidInternalNetRefIndex;
 }
 
 FInternalNetRefIndex FNetRefHandleManager::GetInternalIndexFromNetHandle(FNetHandle Handle) const
 {
 	const FInternalNetRefIndex* InternalIndex = NetHandleToInternalIndex.Find(Handle);
-	return InternalIndex ? *InternalIndex : FNetRefHandleManager::InvalidInternalIndex;
+	return InternalIndex ? *InternalIndex : InvalidInternalNetRefIndex;
 }
 
 FNetRefHandle FNetRefHandleManager::GetNetRefHandleFromInternalIndex(FInternalNetRefIndex InternalIndex) const
@@ -709,7 +707,7 @@ EIrisAsyncLoadingPriority FNetRefHandleManager::GetIrisAsyncLoadingPriority(FInt
 	const FNetRefHandleManager::FReplicatedObjectData& ObjectData = GetReplicatedObjectDataNoCheck(InternalIndex);
 	
 	// If it's a subobject, take the priority of the RootObject
-	if (ObjectData.SubObjectRootIndex != FNetRefHandleManager::InvalidInternalIndex)
+	if (ObjectData.SubObjectRootIndex != InvalidInternalNetRefIndex)
 	{
 		return GetReplicatedObjectDataNoCheck(ObjectData.SubObjectRootIndex).IrisAsyncLoadingPriority;
 	}
@@ -724,7 +722,7 @@ bool FNetRefHandleManager::IsNetRefHandleAssigned(FNetRefHandle Handle) const
 
 bool FNetRefHandleManager::IsLocal(FInternalNetRefIndex InternalIndex) const
 {
-	return (InternalIndex != InvalidInternalIndex) && (ReplicatedObjectStateBuffers[InternalIndex] != nullptr);
+	return (InternalIndex != InvalidInternalNetRefIndex) && (ReplicatedObjectStateBuffers[InternalIndex] != nullptr);
 }
 
 bool FNetRefHandleManager::IsLocalNetRefHandle(FNetRefHandle Handle) const
@@ -746,7 +744,7 @@ bool FNetRefHandleManager::IsRemoteNetRefHandle(FNetRefHandle Handle) const
 bool FNetRefHandleManager::IsRootObject(FNetRefHandle Handle) const
 {
 	const uint32 InternalIndex = GetInternalIndex(Handle);
-	return (InternalIndex != InvalidInternalIndex) ? IsRootObject(InternalIndex) : false;
+	return (InternalIndex != InvalidInternalNetRefIndex) ? IsRootObject(InternalIndex) : false;
 }
 
 uint32 FNetRefHandleManager::GetOriginalDestroyedStartupObjectIndex(FInternalNetRefIndex InternalIndex) const
diff --git a/Engine/Source/Runtime/Net/Iris/Private/Iris/ReplicationSystem/ObjectPollFrequencyLimiter.h b/Engine/Source/Runtime/Net/Iris/Private/Iris/ReplicationSystem/ObjectPollFrequencyLimiter.h
--- a/Engine/Source/Runtime/Net/Iris/Private/Iris/ReplicationSystem/ObjectPollFrequencyLimiter.h
+++ b/Engine/Source/Runtime/Net/Iris/Private/Iris/ReplicationSystem/ObjectPollFrequencyLimiter.h
@@ -6,7 +6,7 @@
 #include "Containers/Array.h"
 #include "Net/Core/NetBitArray.h"
 
-namespace UE::Net::Private
+namespace UE::Net
 {
 	typedef uint32 FInternalNetRefIndex;
 }
diff --git a/Engine/Source/Runtime/Net/Iris/Private/Iris/ReplicationSystem/ObjectReplicationBridge.cpp b/Engine/Source/Runtime/Net/Iris/Private/Iris/ReplicationSystem/ObjectReplicationBridge.cpp
--- a/Engine/Source/Runtime/Net/Iris/Private/Iris/ReplicationSystem/ObjectReplicationBridge.cpp
+++ b/Engine/Source/Runtime/Net/Iris/Private/Iris/ReplicationSystem/ObjectReplicationBridge.cpp
@@ -330,7 +330,7 @@ UNetObjectFactory* UObjectReplicationBridge::GetNetFactory(UE::Net::FNetObjectFa
 	return nullptr;
 }
 
-void UObjectReplicationBridge::OnMaxInternalNetRefIndexIncreased(UE::Net::Private::FInternalNetRefIndex NewMaxInternalIndex)
+void UObjectReplicationBridge::OnMaxInternalNetRefIndexIncreased(UE::Net::FInternalNetRefIndex NewMaxInternalIndex)
 {
 	PollFrequencyLimiter->OnMaxInternalNetRefIndexIncreased(NewMaxInternalIndex);
 	ObjectsWithObjectReferences.SetNumBits(NewMaxInternalIndex);
@@ -481,7 +481,7 @@ void UObjectReplicationBridge::RemoveDestructionInfosForGroup(UE::Net::FNetObjec
 		// Destroy NetObjects and remove from WorldLocations
 		for (const TPair<FNetRefHandle, UE::Net::Private::FStaticDestructionInfo>& It : StaticObjectsPendingDestroy)
 		{
-			const UE::Net::Private::FInternalNetRefIndex InternalObjectIndex = It.Value.InternalReplicationIndex;
+			const UE::Net::FInternalNetRefIndex InternalObjectIndex = It.Value.InternalReplicationIndex;
 			NetRefHandleManager->DestroyNetObjectByIndex(InternalObjectIndex);
 			WorldLocations.RemoveObjectInfoCache(InternalObjectIndex);
 		}
@@ -490,12 +490,12 @@ void UObjectReplicationBridge::RemoveDestructionInfosForGroup(UE::Net::FNetObjec
 	}
 }
 
-void UObjectReplicationBridge::RemoveDestructionInfo(UE::Net::Private::FInternalNetRefIndex InternalObjectIndex)
+void UObjectReplicationBridge::RemoveDestructionInfo(UE::Net::FInternalNetRefIndex InternalObjectIndex)
 {
 	using namespace UE::Net;
 	using namespace UE::Net::Private;
 
-	if (InternalObjectIndex == FNetRefHandleManager::InvalidInternalIndex)
+	if (InternalObjectIndex == InvalidInternalNetRefIndex)
 	{
 		return;
 	}
@@ -538,7 +538,7 @@ void UObjectReplicationBridge::ReadAndExecuteDestructionInfoFromRemote(FReplicat
 	
 	// Look if this object has an internal index assigned. If it doesn't have one it's because it was never replicated to the client before
 	const FInternalNetRefIndex InternalReplicationIndex = NetRefHandleManager->GetInternalIndex(ReferenceToDestroy.GetRefHandle());
-	if (InternalReplicationIndex != FNetRefHandleManager::InvalidInternalIndex)
+	if (InternalReplicationIndex != InvalidInternalNetRefIndex)
 	{
 		NetRefHandleManager->GetReplicatedObjectDataNoCheck(InternalReplicationIndex).bPendingEndReplication = 1U;
 	}
@@ -587,7 +587,7 @@ UE::Net::FNetRefHandle UObjectReplicationBridge::GetReplicatedRefHandle(FNetHand
 
 	// If the object is replicated by the owning ReplicationSystem the internal handle should be valid.
 	FInternalNetRefIndex ObjectInternalIndex = NetRefHandleManager->GetInternalIndexFromNetHandle(Handle);
-	if (ObjectInternalIndex == FNetRefHandleManager::InvalidInternalIndex)
+	if (ObjectInternalIndex == InvalidInternalNetRefIndex)
 	{
 		return FNetRefHandle::GetInvalid();
 	}
@@ -1068,12 +1068,12 @@ UE::Net::EDetachReason ConvertDetachReasonFromInternalToPublic(UE::Net::Private:
 	return EDetachReason::Stopped;
 }
 
-void UObjectReplicationBridge::DestroyNetObjectFromRemote(UE::Net::Private::FInternalNetRefIndex InternalObjectIndex, EReplicationBridgeDestroyInstanceReason DestroyReason, UE::Net::Private::EInternalDetachReason InternalDetachReason)
+void UObjectReplicationBridge::DestroyNetObjectFromRemote(UE::Net::FInternalNetRefIndex InternalObjectIndex, EReplicationBridgeDestroyInstanceReason DestroyReason, UE::Net::Private::EInternalDetachReason InternalDetachReason)
 {
 	using namespace UE::Net;
 	using namespace UE::Net::Private;
 
-	if (InternalObjectIndex == FNetRefHandleManager::InvalidInternalIndex)
+	if (InternalObjectIndex == InvalidInternalNetRefIndex)
 	{
 		return;
 	}
@@ -1166,6 +1166,7 @@ void UObjectReplicationBridge::SetSubObjectDormancyStatus(FNetRefHandle SubObjec
 
 void UObjectReplicationBridge::SetSubObjectNetCondition(FNetRefHandle SubObjectRefHandle, ELifetimeCondition Condition)
 {
+	using namespace UE::Net;
 	using namespace UE::Net::Private;
 
 	// We assume that we can store the condition in an int8;
@@ -1209,6 +1210,7 @@ UE::Net::FNetRefHandle UObjectReplicationBridge::GetRootObjectOfSubObject(FNetRe
 
 void UObjectReplicationBridge::AddDependentObject(FNetRefHandle ParentHandle, FNetRefHandle DependentHandle, UE::Net::EDependentObjectSchedulingHint SchedulingHint)
 {
+	using namespace UE::Net;
 	using namespace UE::Net::Private;
 
 	FReplicationSystemInternal* ReplicationSystemInternal = GetReplicationSystem()->GetReplicationSystemInternal();
@@ -1230,6 +1232,7 @@ void UObjectReplicationBridge::AddDependentObject(FNetRefHandle ParentHandle, FN
 
 void UObjectReplicationBridge::RemoveDependentObject(FNetRefHandle ParentHandle, FNetRefHandle DependentHandle)
 {
+	using namespace UE::Net;
 	using namespace UE::Net::Private;
 
 	UE_LOG_OBJECTREPLICATIONBRIDGE(Verbose, TEXT("RemoveDependentObject Removing dependent object %s from parent %s"), *DependentHandle.ToString(), *ParentHandle.ToString());
@@ -1241,7 +1244,7 @@ void UObjectReplicationBridge::RemoveDependentObject(FNetRefHandle ParentHandle,
 	LocalNetRefHandleManager.RemoveDependentObject(ParentHandle, DependentHandle);
 
 	const FInternalNetRefIndex DependentInternalIndex = LocalNetRefHandleManager.GetInternalIndex(DependentHandle);
-	if (DependentInternalIndex != FNetRefHandleManager::InvalidInternalIndex)
+	if (DependentInternalIndex != InvalidInternalNetRefIndex)
 	{
 		FReplicationFiltering& Filtering = ReplicationSystemInternal->GetFiltering();
 		Filtering.NotifyRemovedDependentObject(DependentInternalIndex);
@@ -1279,7 +1282,7 @@ void UObjectReplicationBridge::RemoveCreationDependencyLink(FNetRefHandle Parent
 	LocalNetRefHandleManager.RemoveCreationDependency(Parent, Child);
 }
 
-bool UObjectReplicationBridge::WriteNetRefHandleCreationInfo(FReplicationBridgeSerializationContext& Context, FNetRefHandle Handle, UE::Net::Private::FInternalNetRefIndex InternalObjectIndex)
+bool UObjectReplicationBridge::WriteNetRefHandleCreationInfo(FReplicationBridgeSerializationContext& Context, FNetRefHandle Handle, UE::Net::FInternalNetRefIndex InternalObjectIndex)
 {
 	UE_IRIS_PARALLEL_EXPR(FScopeLock ScopeLock(&WriteNetRefHandleCreationCS));
 	using namespace UE::Net;
@@ -1347,7 +1350,7 @@ bool UObjectReplicationBridge::CacheNetRefHandleCreationInfo(FNetRefHandle Handl
 	}
 
 	const FInternalNetRefIndex ObjectIndex = NetRefHandleManager->GetInternalIndex(Handle);
-	if (ObjectIndex == FNetRefHandleManager::InvalidInternalIndex)
+	if (ObjectIndex == InvalidInternalNetRefIndex)
 	{
 		UE_LOG_BRIDGEID(LogIris, Error, TEXT("CacheNetRefHandleCreationInfo cannot cache creation info for %s: no internal index assigned to it"), *Handle.ToString());
 		ensureMsgf(false, TEXT("CacheNetRefHandleCreationInfo cannot cache creation info for %s: no internal index assigned to it"), *Handle.ToString());
@@ -1660,7 +1663,7 @@ PRAGMA_DISABLE_DEPRECATION_WARNINGS
 
 			const FInternalNetRefIndex RootObjectIndex = DestroyedObjectData.SubObjectRootIndex;
 
-			if (RootObjectIndex != FNetRefHandleManager::InvalidInternalIndex)
+			if (RootObjectIndex != InvalidInternalNetRefIndex)
 			{
 				DestroyContext.RootObject = NetRefHandleManager->GetReplicatedObjectInstance(RootObjectIndex);
 				DestroyContext.RootObjectHandle = NetRefHandleManager->GetNetRefHandleFromInternalIndex(RootObjectIndex);
@@ -1693,13 +1696,14 @@ PRAGMA_ENABLE_DEPRECATION_WARNINGS
 	}
 }
 
-void UObjectReplicationBridge::DetachSubObjectInstancesFromRemote(UE::Net::Private::FInternalNetRefIndex InternalObjectIndex, EReplicationBridgeDestroyInstanceReason DestroyReason, EReplicationBridgeDestroyInstanceFlags DestroyFlags)
+void UObjectReplicationBridge::DetachSubObjectInstancesFromRemote(UE::Net::FInternalNetRefIndex InternalObjectIndex, EReplicationBridgeDestroyInstanceReason DestroyReason, EReplicationBridgeDestroyInstanceFlags DestroyFlags)
 {
+	using namespace UE::Net;
 	using namespace UE::Net::Private;
 
 	ensure(UE::Net::IsUsingNewDetachAPI() == false);
 
-	if (InternalObjectIndex == FNetRefHandleManager::InvalidInternalIndex)
+	if (InternalObjectIndex == UE::Net::InvalidInternalNetRefIndex)
 	{
 		return;
 	}
@@ -1994,7 +1998,7 @@ FReplicationBridgeCreateNetRefHandleResult UObjectReplicationBridge::CreateNetRe
 	return CreateResult;
 }
 
-void UObjectReplicationBridge::SubObjectCreatedFromReplication(UE::Net::Private::FInternalNetRefIndex RootObjectIndex, FNetRefHandle SubObjectCreated)
+void UObjectReplicationBridge::SubObjectCreatedFromReplication(UE::Net::FInternalNetRefIndex RootObjectIndex, FNetRefHandle SubObjectCreated)
 {
 	using namespace UE::Net;
 	using namespace UE::Net::Private;
@@ -2005,7 +2009,7 @@ void UObjectReplicationBridge::SubObjectCreatedFromReplication(UE::Net::Private:
 	Factory->SubObjectCreatedFromReplication(NetRefHandleManager->GetNetRefHandleFromInternalIndex(RootObjectIndex), SubObjectCreated);
 }
 
-void UObjectReplicationBridge::PostApplyInitialState(UE::Net::Private::FInternalNetRefIndex InternalObjectIndex)
+void UObjectReplicationBridge::PostApplyInitialState(UE::Net::FInternalNetRefIndex InternalObjectIndex)
 {
 	using namespace UE::Net;
 	using namespace UE::Net::Private;
@@ -2167,6 +2171,7 @@ UObjectReplicationBridge::FInstancePreUpdateFunction UObjectReplicationBridge::G
 
 void UObjectReplicationBridge::ForcePollObject(FNetRefHandle Handle)
 {
+	using namespace UE::Net;
 	using namespace UE::Net::Private;
 
 	if (!Handle.IsValid())
@@ -2179,7 +2184,7 @@ void UObjectReplicationBridge::ForcePollObject(FNetRefHandle Handle)
 	const FNetRefHandleManager& LocalNetRefHandleManager = GetReplicationSystem()->GetReplicationSystemInternal()->GetNetRefHandleManager();
 
 	const FInternalNetRefIndex ObjectIndex = LocalNetRefHandleManager.GetInternalIndex(Handle);
-	if (ObjectIndex == FNetRefHandleManager::InvalidInternalIndex)
+	if (ObjectIndex == UE::Net::InvalidInternalNetRefIndex)
 	{
 		return;
 	}
@@ -2449,9 +2454,9 @@ void UObjectReplicationBridge::ReconcileNewSubObjects(UE::Net::FNetBitArrayView
 		}
 
 		const FInternalNetRefIndex RootObjectIndex = LocalNetRefHandleManager.GetRootObjectInternalIndexOfSubObject(SubObjectIndex);
-		if (UNLIKELY(RootObjectIndex == FNetRefHandleManager::InvalidInternalIndex))
+		if (UNLIKELY(RootObjectIndex == InvalidInternalNetRefIndex))
 		{
-			ensureMsgf(RootObjectIndex != FNetRefHandleManager::InvalidInternalIndex, TEXT("SubObject %s (Index:%u) had invalid RootObjectIndex"), *GetNameSafe(LocalNetRefHandleManager.GetReplicatedObjectInstance(SubObjectIndex)), SubObjectIndex);
+			ensureMsgf(RootObjectIndex != InvalidInternalNetRefIndex, TEXT("SubObject %s (Index:%u) had invalid RootObjectIndex"), *GetNameSafe(LocalNetRefHandleManager.GetReplicatedObjectInstance(SubObjectIndex)), SubObjectIndex);
 			return;
 		}
 
@@ -2480,7 +2485,7 @@ void UObjectReplicationBridge::ReconcileNewSubObjects(UE::Net::FNetBitArrayView
 	FNetBitArrayView::ForAllSetBits(GlobalScopeList, CurrentFrameScopeList, FNetBitArrayView::AndNotOp, HandleNewSubObject);
 }
 
-void UObjectReplicationBridge::ForceUpdateWorldLocation(FNetRefHandle NetRefHandle, UE::Net::Private::FInternalNetRefIndex InternalObjectIndex)
+void UObjectReplicationBridge::ForceUpdateWorldLocation(FNetRefHandle NetRefHandle, UE::Net::FInternalNetRefIndex InternalObjectIndex)
 {
 	using namespace UE::Net;
 
@@ -2598,13 +2603,14 @@ void UObjectReplicationBridge::UpdateInstancesWorldLocation()
 
 void UObjectReplicationBridge::SetPollWithObject(FNetRefHandle ObjectToPollWithHandle, FNetRefHandle ObjectHandle)
 {
-	const UE::Net::Private::FInternalNetRefIndex PollWithInternalReplicationIndex = NetRefHandleManager->GetInternalIndex(ObjectToPollWithHandle);
-	const UE::Net::Private::FInternalNetRefIndex InternalReplicationIndex = NetRefHandleManager->GetInternalIndex(ObjectHandle);
+	const UE::Net::FInternalNetRefIndex PollWithInternalReplicationIndex = NetRefHandleManager->GetInternalIndex(ObjectToPollWithHandle);
+	const UE::Net::FInternalNetRefIndex InternalReplicationIndex = NetRefHandleManager->GetInternalIndex(ObjectHandle);
 	PollFrequencyLimiter->SetPollWithObject(PollWithInternalReplicationIndex, InternalReplicationIndex);
 }
 
 bool UObjectReplicationBridge::GetObjectWantsToBeDormant(FNetRefHandle Handle) const
 {
+	using namespace UE::Net;
 	using namespace UE::Net::Private;
 
 	const FReplicationSystemInternal* ReplicationSystemInternal = GetReplicationSystem()->GetReplicationSystemInternal();
@@ -2626,7 +2632,7 @@ void UObjectReplicationBridge::SetObjectWantsToBeDormant(FNetRefHandle Handle, b
 	FNetRefHandleManager& LocalNetRefHandleManager = GetReplicationSystem()->GetReplicationSystemInternal()->GetNetRefHandleManager();
 
 	const FInternalNetRefIndex InternalObjectIndex = LocalNetRefHandleManager.GetInternalIndex(Handle);
-	if (InternalObjectIndex == FNetRefHandleManager::InvalidInternalIndex)
+	if (InternalObjectIndex == InvalidInternalNetRefIndex)
 	{
 		return;
 	}
@@ -2663,7 +2669,7 @@ void UObjectReplicationBridge::NetFlushDormantObject(FNetRefHandle Handle)
 	FNetRefHandleManager& LocalNetRefHandleManager = GetReplicationSystem()->GetReplicationSystemInternal()->GetNetRefHandleManager();
 
 	const FInternalNetRefIndex InternalObjectIndex = LocalNetRefHandleManager.GetInternalIndex(Handle);
-	if (InternalObjectIndex == FNetRefHandleManager::InvalidInternalIndex)
+	if (InternalObjectIndex == InvalidInternalNetRefIndex)
 	{
 		return;
 	}
@@ -3394,7 +3400,7 @@ uint8 UObjectReplicationBridge::ConvertPollFrequencyIntoFrames(float PollFrequen
 	return static_cast<uint8>(FMath::Clamp<uint32>(FramesBetweenUpdatesForObject, 0U, UE::Net::Private::FObjectPollFrequencyLimiter::GetMaxPollingFrames()));
 }
 
-float UObjectReplicationBridge::GetPollFrequencyOfRootObject(UE::Net::Private::FInternalNetRefIndex RootObjectIndex) const
+float UObjectReplicationBridge::GetPollFrequencyOfRootObject(UE::Net::FInternalNetRefIndex RootObjectIndex) const
 {
 	using namespace UE::Net;
 	using namespace UE::Net::Private;
@@ -3458,7 +3464,7 @@ void UObjectReplicationBridge::SetPollFrequency(FNetRefHandle RefHandle, float P
 	const FNetRefHandleManager& LocalNetRefHandleManager = ReplicationSystemInternal->GetNetRefHandleManager();
 
 	FInternalNetRefIndex RootObjectIndex = LocalNetRefHandleManager.GetInternalIndex(RefHandle);
-	if (RootObjectIndex == FNetRefHandleManager::InvalidInternalIndex)
+	if (RootObjectIndex == InvalidInternalNetRefIndex)
 	{
 		return;
 	}
@@ -3486,7 +3492,7 @@ void UObjectReplicationBridge::OnProtocolMismatchReported(FNetRefHandle RefHandl
 	
 
 	const FInternalNetRefIndex ObjectInternalIndex = NetRefHandleManager->GetInternalIndex(RefHandle);
-	if (ObjectInternalIndex == FNetRefHandleManager::InvalidInternalIndex)
+	if (ObjectInternalIndex == InvalidInternalNetRefIndex)
 	{
 		UE_LOG_BRIDGEID(LogIris, Warning, TEXT("OnProtocolMismatchReported from Connection:%s for %s. But object has no InternalIndex."), *PrintConnectionInfo(ConnectionId), *RefHandle.ToString());
 		return;
@@ -3576,7 +3582,7 @@ void UObjectReplicationBridge::OnErrorBitstreamCorrupted(FNetRefHandle RefHandle
 	using namespace UE::Net::Private;
 
 	const FInternalNetRefIndex ObjectInternalIndex = NetRefHandleManager->GetInternalIndex(RefHandle);
-	if (ObjectInternalIndex == FNetRefHandleManager::InvalidInternalIndex)
+	if (ObjectInternalIndex == InvalidInternalNetRefIndex)
 	{
 		UE_LOG_BRIDGEID(LogIris, Error, TEXT("Bitstream corruption from unmapped handle: %s. Reported by client: %s."), *RefHandle.ToString(), *PrintConnectionInfo(ConnectionId));
 	}
@@ -3594,7 +3600,7 @@ void UObjectReplicationBridge::OnErrorCannotReplicateObject(FNetRefHandle RefHan
 	using namespace UE::Net::Private;
 
 	const FInternalNetRefIndex ObjectInternalIndex = NetRefHandleManager->GetInternalIndex(RefHandle);
-	if (ObjectInternalIndex == FNetRefHandleManager::InvalidInternalIndex)
+	if (ObjectInternalIndex == InvalidInternalNetRefIndex)
 	{
 		UE_LOG_BRIDGEID(LogIris, Error, TEXT("Client cannot replicate unmapped handle: %s. Reported by client: %s."), *RefHandle.ToString(), *PrintConnectionInfo(ConnectionId));
 	}
@@ -3624,7 +3630,7 @@ void UObjectReplicationBridge::OnErrorBlockedByAsyncLoading(FNetRefHandle RefHan
 
 	const FInternalNetRefIndex ObjectInternalIndex = NetRefHandleManager->GetInternalIndex(RefHandle);
 
-	if (ObjectInternalIndex == FNetRefHandleManager::InvalidInternalIndex)
+	if (ObjectInternalIndex == InvalidInternalNetRefIndex)
 	{
 		UE_LOG_BRIDGEID(LogIris, Error, TEXT("Replication Blocked by AsyncLoading for unmapped handle: %s. Client: %s"), *RefHandle.ToString(), *PrintConnectionInfo(ConnectionId));
 	}
@@ -3659,7 +3665,7 @@ void UObjectReplicationBridge::OnErrorBlockedByParents(FNetRefHandle RefHandle,
 	// A client cannot instantiate the Object because it has been waiting on the creation parents to be spawned for too long (default is 30secs)
 	const FInternalNetRefIndex ObjectInternalIndex = NetRefHandleManager->GetInternalIndex(RefHandle);
 	
-	if (ObjectInternalIndex == FNetRefHandleManager::InvalidInternalIndex)
+	if (ObjectInternalIndex == InvalidInternalNetRefIndex)
 	{
 		UE_LOG_BRIDGEID(LogIris, Error, TEXT("Client cannot spawn unmapped object: %s. Reported by client: %s"), *RefHandle.ToString(), *PrintConnectionInfo(ConnectionId));
 	}
diff --git a/Engine/Source/Runtime/Net/Iris/Private/Iris/ReplicationSystem/ObjectReplicationBridgeDebugging.cpp b/Engine/Source/Runtime/Net/Iris/Private/Iris/ReplicationSystem/ObjectReplicationBridgeDebugging.cpp
--- a/Engine/Source/Runtime/Net/Iris/Private/Iris/ReplicationSystem/ObjectReplicationBridgeDebugging.cpp
+++ b/Engine/Source/Runtime/Net/Iris/Private/Iris/ReplicationSystem/ObjectReplicationBridgeDebugging.cpp
@@ -256,7 +256,7 @@ FString PrintNetObject(FNetRefHandleManager* NetRefHandleManager, FInternalNetRe
 	UObject* ObjectPtr = NetRefHandleManager->GetReplicatedObjectInstance(ObjectIndex);
 
 	return FString::Printf(TEXT("%s %s (InternalIndex: %u) (%s)"), 
-		(NetObjectData.SubObjectRootIndex == FNetRefHandleManager::InvalidInternalIndex) ? TEXT("RootObject"):TEXT("SubObject"),
+		(NetObjectData.SubObjectRootIndex == InvalidInternalNetRefIndex) ? TEXT("RootObject"):TEXT("SubObject"),
 		*GetNameSafe(ObjectPtr), ObjectIndex, *NetRefHandle.ToString()
 	);
 }
diff --git a/Engine/Source/Runtime/Net/Iris/Private/Iris/ReplicationSystem/Prioritization/NetObjectConnectionFilter.cpp b/Engine/Source/Runtime/Net/Iris/Private/Iris/ReplicationSystem/Prioritization/NetObjectConnectionFilter.cpp
--- a/Engine/Source/Runtime/Net/Iris/Private/Iris/ReplicationSystem/Prioritization/NetObjectConnectionFilter.cpp
+++ b/Engine/Source/Runtime/Net/Iris/Private/Iris/ReplicationSystem/Prioritization/NetObjectConnectionFilter.cpp
@@ -17,9 +17,9 @@ void UNetObjectConnectionFilter::SetReplicateToConnection(UE::Net::FNetRefHandle
 		return;
 	}
 
-	const UE::Net::Private::FInternalNetRefIndex ObjectIndex = GetObjectIndex(RefHandle);
+	const UE::Net::FInternalNetRefIndex ObjectIndex = GetObjectIndex(RefHandle);
 
-	if (ObjectIndex == UE::Net::Private::FNetRefHandleManager::InvalidInternalIndex)
+	if (ObjectIndex == UE::Net::InvalidInternalNetRefIndex)
 	{
 		return;
 	}
@@ -68,7 +68,7 @@ void UNetObjectConnectionFilter::RemoveConnection(uint32 ConnectionId)
 	ConnInfo = FPerConnectionInfo();
 }
 
-bool UNetObjectConnectionFilter::AddObject(uint32 ObjectIndex, FNetObjectFilterAddObjectParams& Params)
+bool UNetObjectConnectionFilter::AddObject(UE::Net::FInternalNetRefIndex ObjectIndex, FNetObjectFilterAddObjectParams& Params)
 {
 	const uint32 LocalIndex = UsedLocalInfoIndices.FindFirstZero();
 	if (LocalIndex == UE::Net::FNetBitArrayBase::InvalidIndex)
@@ -88,7 +88,7 @@ bool UNetObjectConnectionFilter::AddObject(uint32 ObjectIndex, FNetObjectFilterA
 	return true;
 }
 
-void UNetObjectConnectionFilter::RemoveObject(uint32 ObjectIndex, const FNetObjectFilteringInfo& InInfo)
+void UNetObjectConnectionFilter::RemoveObject(UE::Net::FInternalNetRefIndex ObjectIndex, const FNetObjectFilteringInfo& InInfo)
 {
 	bObjectRemoved = true;
 
@@ -100,7 +100,7 @@ void UNetObjectConnectionFilter::RemoveObject(uint32 ObjectIndex, const FNetObje
 	UE_LOG(LogIris, Verbose, TEXT("UNetObjectConnectionFilter::RemoveObject removed %u | Mapped to LocalIndex %u"), ObjectIndex, LocalIndex);
 
 	// For good measure. It's not strictly needed.
-	LocalToNetRefIndex[LocalIndex] = UE::Net::Private::FNetRefHandleManager::InvalidInternalIndex;
+	LocalToNetRefIndex[LocalIndex] = UE::Net::InvalidInternalNetRefIndex;
 }
 
 void UNetObjectConnectionFilter::PreFilter(FNetObjectPreFilteringParams& Params)
diff --git a/Engine/Source/Runtime/Net/Iris/Private/Iris/ReplicationSystem/Prioritization/ReplicationPrioritization.cpp b/Engine/Source/Runtime/Net/Iris/Private/Iris/ReplicationSystem/Prioritization/ReplicationPrioritization.cpp
--- a/Engine/Source/Runtime/Net/Iris/Private/Iris/ReplicationSystem/Prioritization/ReplicationPrioritization.cpp
+++ b/Engine/Source/Runtime/Net/Iris/Private/Iris/ReplicationSystem/Prioritization/ReplicationPrioritization.cpp
@@ -102,7 +102,7 @@ public:
 
 	struct FPerPrioritizerInfo
 	{
-		TChunkedArrayWithChunkManagement<uint32, MaxObjectCountPerBatch*sizeof(uint32)> ObjectIndices;
+		TChunkedArrayWithChunkManagement<FInternalNetRefIndex, MaxObjectCountPerBatch*sizeof(uint32)> ObjectIndices;
 	};
 
 	explicit FPrioritizerBatchHelper(uint32 PrioritizerCount)
@@ -221,7 +221,7 @@ public:
 
 	struct FPerPrioritizerInfo
 	{
-		uint32* ObjectIndices;
+		FInternalNetRefIndex* ObjectIndices;
 		uint32 ObjectCount;
 
 		FReplicationInstanceProtocol const** InstanceProtocols;
@@ -277,7 +277,7 @@ private:
 		}
 	}
 
-	TArray<uint32> ObjectIndicesStorage;
+	TArray<UE::Net::FInternalNetRefIndex> ObjectIndicesStorage;
 	TArray<const FReplicationInstanceProtocol*> InstanceProtocolsStorage;
 	const FNetRefHandleManager* NetRefHandleManager;
 };
@@ -384,7 +384,7 @@ void FReplicationPrioritization::OnMaxInternalNetRefIndexIncreased(FInternalNetR
 	}
 }
 
-void FReplicationPrioritization::SetStaticPriority(uint32 ObjectIndex, float NewPrio)
+void FReplicationPrioritization::SetStaticPriority(FInternalNetRefIndex ObjectIndex, float NewPrio)
 {
 	if (!ensureMsgf(NewPrio >= 0.0f, TEXT("Trying to set invalid priority %f"), NewPrio))
 	{
@@ -411,7 +411,7 @@ void FReplicationPrioritization::SetStaticPriority(uint32 ObjectIndex, float New
 	}
 }
 
-bool FReplicationPrioritization::SetPrioritizer(uint32 ObjectIndex, FNetObjectPrioritizerHandle NewPrioritizer)
+bool FReplicationPrioritization::SetPrioritizer(FInternalNetRefIndex ObjectIndex, FNetObjectPrioritizerHandle NewPrioritizer)
 {
 	if (!ensureMsgf(NewPrioritizer != InvalidNetObjectPrioritizerHandle, TEXT("%s"), TEXT("Call SetStaticPriority if you want to use a static priority for the object.")))
 	{
@@ -674,16 +674,16 @@ void FReplicationPrioritization::UpdatePrioritiesForNewAndDeletedObjects()
 		DefaultPriorities[ObjectIndex] = DefaultPriority;
 	};
 
-	TArray<uint32> NewIndices;
-	TFunction<void(uint32)> DoNothing = [](uint32 ObjectIndex){};
-	TFunction<void(uint32)> AddIndexAndClearFromNewPriority = [&NewIndices, this](uint32 ObjectIndex)
+	TArray<FInternalNetRefIndex> NewIndices;
+	TFunction<void(FInternalNetRefIndex)> DoNothing = [](FInternalNetRefIndex ObjectIndex){};
+	TFunction<void(FInternalNetRefIndex)> AddIndexAndClearFromNewPriority = [&NewIndices, this](FInternalNetRefIndex ObjectIndex)
 	{
 		NewIndices.Add(ObjectIndex);
 		// Prevent the same index from being added twice
 		this->ObjectsWithNewStaticPriority.ClearBit(ObjectIndex);
 	};
 
-	TFunction<void(uint32)> ForEachNewObject = (ConnectionCount > 0 ? AddIndexAndClearFromNewPriority : DoNothing);
+	TFunction<void(FInternalNetRefIndex)> ForEachNewObject = (ConnectionCount > 0 ? AddIndexAndClearFromNewPriority : DoNothing);
 	if (ConnectionCount > 0)
 	{
 		NewIndices.Reserve(FMath::Min(1024U, MaxInternalNetRefIndex));
@@ -697,7 +697,7 @@ void FReplicationPrioritization::UpdatePrioritiesForNewAndDeletedObjects()
 
 		if (ConnectionCount > 0)
 		{
-			ObjectsWithNewStaticPriority.ForAllSetBits([this, &NewIndices](uint32 ObjectIndex)
+			ObjectsWithNewStaticPriority.ForAllSetBits([this, &NewIndices](FInternalNetRefIndex ObjectIndex)
 			{ 
 				NewIndices.Add(ObjectIndex); 
 			});
@@ -710,7 +710,7 @@ void FReplicationPrioritization::UpdatePrioritiesForNewAndDeletedObjects()
 	// Copy the priorities for the new objects to each connection
 	if (NewIndices.Num() > 0 && ConnectionCount > 0)
 	{
-		const TArrayView<uint32> ObjectIndices = MakeArrayView(NewIndices);
+		const TArrayView<FInternalNetRefIndex> ObjectIndices = MakeArrayView(NewIndices);
 		for (FPerConnectionInfo& ConnectionInfo : ConnectionInfos)
 		{
 			if (!ConnectionInfo.IsValid)
@@ -718,7 +718,7 @@ void FReplicationPrioritization::UpdatePrioritiesForNewAndDeletedObjects()
 				continue;
 			}
 
-			for (const uint32 ObjectIndex : ObjectIndices)
+			for (const FInternalNetRefIndex ObjectIndex : ObjectIndices)
 			{
 				ConnectionInfo.Priorities[ObjectIndex] = DefaultPriorities[ObjectIndex];
 			}
@@ -825,7 +825,7 @@ void FReplicationPrioritization::SetHighPriorityOnViewTargets(const TArrayView<f
 	for (FNetHandle NetHandle : ViewTargets)
 	{
 		const FInternalNetRefIndex ViewTargetInternalIndex = NetRefHandleManager->GetInternalIndexFromNetHandle(NetHandle);
-		if (ViewTargetInternalIndex != FNetRefHandleManager::InvalidInternalIndex)
+		if (ViewTargetInternalIndex != InvalidInternalNetRefIndex)
 		{
 			Priorities[ViewTargetInternalIndex] = ViewTargetHighPriority;
 		}
@@ -842,7 +842,7 @@ void FReplicationPrioritization::NotifyPrioritizersOfDirtyObjects(const FNetBitA
 	FUpdateDirtyObjectsBatchHelper BatchHelper(NetRefHandleManager, PrioritizerInfos.Num());
 
 	constexpr SIZE_T MaxBatchObjectCount = FUpdateDirtyObjectsBatchHelper::Constants::MaxObjectCountPerBatch;
-	uint32 ObjectIndices[MaxBatchObjectCount];
+	FInternalNetRefIndex ObjectIndices[MaxBatchObjectCount];
 
 	const uint32 BitCount = ~0U;
 	for (uint32 ObjectCount, StartIndex = 0; (ObjectCount = DirtyObjectsThisFrame.GetSetBitIndices(StartIndex, BitCount, ObjectIndices, MaxBatchObjectCount)) > 0; )
@@ -857,7 +857,7 @@ void FReplicationPrioritization::NotifyPrioritizersOfDirtyObjects(const FNetBitA
 	}
 }
 
-void FReplicationPrioritization::BatchNotifyPrioritizersOfDirtyObjects(FUpdateDirtyObjectsBatchHelper& BatchHelper, uint32* ObjectIndices, uint32 ObjectCount)
+void FReplicationPrioritization::BatchNotifyPrioritizersOfDirtyObjects(FUpdateDirtyObjectsBatchHelper& BatchHelper, FInternalNetRefIndex* ObjectIndices, uint32 ObjectCount)
 {
 	BatchHelper.PrepareBatch(ObjectIndices, ObjectCount, ObjectIndexToPrioritizer.GetData());
 
diff --git a/Engine/Source/Runtime/Net/Iris/Private/Iris/ReplicationSystem/Prioritization/ReplicationPrioritization.h b/Engine/Source/Runtime/Net/Iris/Private/Iris/ReplicationSystem/Prioritization/ReplicationPrioritization.h
--- a/Engine/Source/Runtime/Net/Iris/Private/Iris/ReplicationSystem/Prioritization/ReplicationPrioritization.h
+++ b/Engine/Source/Runtime/Net/Iris/Private/Iris/ReplicationSystem/Prioritization/ReplicationPrioritization.h
@@ -14,12 +14,13 @@ class UReplicationSystem;
 namespace UE::Net
 {
 	class FNetBitArrayView;
+	
+	typedef uint32 FInternalNetRefIndex;
+
 	namespace Private
 	{
 		class FNetRefHandleManager;
 		class FReplicationConnections;
-
-		typedef uint32 FInternalNetRefIndex;
 	}
 
 	// For testing
@@ -50,8 +51,8 @@ public:
 
 	void Prioritize(const FNetBitArrayView& ConnectionsToSend, const FNetBitArrayView& DirtyObjectsThisFrame);
 
-	void SetStaticPriority(uint32 ObjectIndex, float Prio);
-	bool SetPrioritizer(uint32 ObjectIndex, FNetObjectPrioritizerHandle Prioritizer);
+	void SetStaticPriority(FInternalNetRefIndex ObjectIndex, float Prio);
+	bool SetPrioritizer(FInternalNetRefIndex ObjectIndex, FNetObjectPrioritizerHandle Prioritizer);
 	FNetObjectPrioritizerHandle GetPrioritizerHandle(const FName PrioritizerName) const;
 	UNetObjectPrioritizer* GetPrioritizer(const FName PrioritizerName) const;
 
@@ -74,7 +75,7 @@ private:
 	void PrioritizeForConnection(uint32 ConnId, FPrioritizerBatchHelper& BatchHelper, FNetBitArrayView Objects);
 	void SetHighPriorityOnViewTargets(const TArrayView<float>& Priorities, const FReplicationView& View);
 	void NotifyPrioritizersOfDirtyObjects(const FNetBitArrayView& DirtyObjectsThisFrame);
-	void BatchNotifyPrioritizersOfDirtyObjects(FUpdateDirtyObjectsBatchHelper& BatchHelper, uint32* ObjectIndices, uint32 ObjectCount);
+	void BatchNotifyPrioritizersOfDirtyObjects(FUpdateDirtyObjectsBatchHelper& BatchHelper, FInternalNetRefIndex* ObjectIndices, uint32 ObjectCount);
 	void InitPrioritizers();
 
 private:
diff --git a/Engine/Source/Runtime/Net/Iris/Private/Iris/ReplicationSystem/ReplicationBridge.cpp b/Engine/Source/Runtime/Net/Iris/Private/Iris/ReplicationSystem/ReplicationBridge.cpp
--- a/Engine/Source/Runtime/Net/Iris/Private/Iris/ReplicationSystem/ReplicationBridge.cpp
+++ b/Engine/Source/Runtime/Net/Iris/Private/Iris/ReplicationSystem/ReplicationBridge.cpp
@@ -15,6 +15,7 @@
 
 #include "Iris/ReplicationSystem/ChangeMaskCache.h"
 #include "Iris/ReplicationSystem/LegacyPushModel.h"
+#include "Iris/ReplicationSystem/InternalNetRefIndexManager.h"
 #include "Iris/ReplicationSystem/NetRefHandleManagerTypes.h"
 #include "Iris/ReplicationSystem/ReplicationFragment.h"
 #include "Iris/ReplicationSystem/ReplicationFragmentInternal.h"
@@ -160,7 +161,7 @@ void UReplicationBridge::Deinitialize()
 	// Detach all replicated instances that have not yet been destroyed as part of shutting down the rest of the game.
 	NetRefHandleManager->GetAssignedInternalIndices().ForAllSetBits([this](uint32 InternalObjectIndex) 
 	{
-		if (InternalObjectIndex == FNetRefHandleManager::InvalidInternalIndex)
+		if (InternalObjectIndex == InvalidInternalNetRefIndex)
 		{
 			return;
 		}
@@ -233,7 +234,7 @@ void UReplicationBridge::InternalAttachInstanceToNetRefHandle(FNetRefHandle RefH
 	}
 }
 
-void UReplicationBridge::InternalDetachInstanceFromNetRefHandle(UE::Net::Private::FInternalNetRefIndex InternalObjectIndex, bool bUnBindInstanceProtocol)
+void UReplicationBridge::InternalDetachInstanceFromNetRefHandle(UE::Net::FInternalNetRefIndex InternalObjectIndex, bool bUnBindInstanceProtocol)
 {
 	using namespace UE::Net;
 	using namespace UE::Net::Private;
@@ -250,12 +251,12 @@ void UReplicationBridge::InternalDetachInstanceFromNetRefHandle(UE::Net::Private
 	}
 }
 
-void UReplicationBridge::InternalDestroyNetObject(UE::Net::Private::FInternalNetRefIndex InternalObjectIndex)
+void UReplicationBridge::InternalDestroyNetObject(UE::Net::FInternalNetRefIndex InternalObjectIndex)
 {
 	using namespace UE::Net;
 	using namespace UE::Net::Private;
 
-	if (InternalObjectIndex == FNetRefHandleManager::InvalidInternalIndex)
+	if (InternalObjectIndex == InvalidInternalNetRefIndex)
 	{
 		return;
 	}
@@ -271,7 +272,7 @@ void UReplicationBridge::InternalDestroyNetObject(UE::Net::Private::FInternalNet
 	NetRefHandleManager->DestroyNetObjectByIndex(InternalObjectIndex);
 }
 
-void UReplicationBridge::InternalStartDestroyLocalObject(UE::Net::Private::FInternalNetRefIndex InternalReplicationIndex, EEndReplicationFlags EndReplicationFlags)
+void UReplicationBridge::InternalStartDestroyLocalObject(UE::Net::FInternalNetRefIndex InternalReplicationIndex, EEndReplicationFlags EndReplicationFlags)
 {
 	UE_LOG_REPLICATIONBRIDGE(Verbose, TEXT("InternalStartDestroyLocalObject for %s | EndReplicationFlags: %s"), *NetRefHandleManager->PrintObjectFromIndex(InternalReplicationIndex), *LexToString(EndReplicationFlags));
 
@@ -307,7 +308,7 @@ void UReplicationBridge::InternalStartDestroyLocalObject(UE::Net::Private::FInte
 	InternalStartDestroyLocalSubObjects(InternalReplicationIndex, EndReplicationFlags);
 }
 
-void UReplicationBridge::InternalFinishDestroyLocalObject(UE::Net::Private::FInternalNetRefIndex InternalReplicationIndex, EEndReplicationFlags EndReplicationFlags)
+void UReplicationBridge::InternalFinishDestroyLocalObject(UE::Net::FInternalNetRefIndex InternalReplicationIndex, EEndReplicationFlags EndReplicationFlags)
 {
 	// If the object are a member of any groups we need to remove it to make sure that we update filtering
 	GetReplicationSystem()->GetReplicationSystemInternal()->RemoveFromAllGroups(InternalReplicationIndex);
@@ -321,10 +322,10 @@ void UReplicationBridge::InternalFinishDestroyLocalObject(UE::Net::Private::FInt
 
 void UReplicationBridge::DestroyLocalNetHandle(FNetRefHandle Handle, EEndReplicationFlags EndReplicationFlags)
 {
-	using namespace UE::Net::Private;
+	using namespace UE::Net;
 
 	const FInternalNetRefIndex InternalReplicationIndex = NetRefHandleManager->GetInternalIndex(Handle);
-	if (InternalReplicationIndex != FNetRefHandleManager::InvalidInternalIndex)
+	if (InternalReplicationIndex != InvalidInternalNetRefIndex)
 	{
 		UE_LOG_REPLICATIONBRIDGE(Verbose, TEXT("DestroyLocalNetHandle for %s | EndReplicationFlags: %s"), *NetRefHandleManager->PrintObjectFromNetRefHandle(Handle), *LexToString(EndReplicationFlags));
 		InternalStartDestroyLocalObject(InternalReplicationIndex, EndReplicationFlags);
@@ -384,12 +385,13 @@ PRAGMA_ENABLE_DEPRECATION_WARNINGS
 	}
 }
 
-void UReplicationBridge::InternalStartDestroyLocalSubObjects(UE::Net::Private::FInternalNetRefIndex OwnerInternalIndex, EEndReplicationFlags Flags)
+void UReplicationBridge::InternalStartDestroyLocalSubObjects(UE::Net::FInternalNetRefIndex OwnerInternalIndex, EEndReplicationFlags Flags)
 {
+	using namespace UE::Net;
 	using namespace UE::Net::Private;
 
 	// Destroy SubObjects
-	if (OwnerInternalIndex != FNetRefHandleManager::InvalidInternalIndex)
+	if (OwnerInternalIndex != InvalidInternalNetRefIndex)
 	{
 		for (FInternalNetRefIndex SubObjectInternalIndex : NetRefHandleManager->GetChildSubObjects(OwnerInternalIndex))
 		{
@@ -403,12 +405,13 @@ void UReplicationBridge::InternalStartDestroyLocalSubObjects(UE::Net::Private::F
 	}
 }
 
-void UReplicationBridge::InternalFinishDestroyLocalSubObjects(UE::Net::Private::FInternalNetRefIndex OwnerInternalIndex, EEndReplicationFlags Flags)
+void UReplicationBridge::InternalFinishDestroyLocalSubObjects(UE::Net::FInternalNetRefIndex OwnerInternalIndex, EEndReplicationFlags Flags)
 {
+	using namespace UE::Net;
 	using namespace UE::Net::Private;
 
 	// Destroy SubObjects
-	if (OwnerInternalIndex != FNetRefHandleManager::InvalidInternalIndex)
+	if (OwnerInternalIndex != InvalidInternalNetRefIndex)
 	{
 		for (FInternalNetRefIndex SubObjectInternalIndex : NetRefHandleManager->GetChildSubObjects(OwnerInternalIndex))
 		{
@@ -419,11 +422,12 @@ void UReplicationBridge::InternalFinishDestroyLocalSubObjects(UE::Net::Private::
 	}
 }
 
-void UReplicationBridge::InternalRequestAsyncStopReplication(UE::Net::Private::FInternalNetRefIndex ObjectInternalIndex, EEndReplicationFlags EndReplicationFlags)
+void UReplicationBridge::InternalRequestAsyncStopReplication(UE::Net::FInternalNetRefIndex ObjectInternalIndex, EEndReplicationFlags EndReplicationFlags)
 {
+	using namespace UE::Net;
 	using namespace UE::Net::Private;
 
-	if (ObjectInternalIndex != FNetRefHandleManager::InvalidInternalIndex && NetRefHandleManager->IsScopableIndex(ObjectInternalIndex))
+	if (ObjectInternalIndex != InvalidInternalNetRefIndex && NetRefHandleManager->IsScopableIndex(ObjectInternalIndex))
 	{
 		UE_LOG_REPLICATIONBRIDGE(Verbose, TEXT("%hs: %s | EndFlags: %s"), __FUNCTION__, *NetRefHandleManager->PrintObjectFromIndex(ObjectInternalIndex), ToCStr(LexToString(EndReplicationFlags)));
 
@@ -450,6 +454,7 @@ void UReplicationBridge::InternalRequestAsyncStopReplication(UE::Net::Private::F
 
 void UReplicationBridge::StopReplicatingNetRefHandle(FNetRefHandle Handle, EEndReplicationFlags EndReplicationFlags)
 {
+	using namespace UE::Net;
 	using namespace UE::Net::Private;
 
 	if (!IsReplicatedHandle(Handle))
@@ -552,7 +557,7 @@ void UReplicationBridge::StopReplicatingNetRefHandle(FNetRefHandle Handle, EEndR
 			return;
 		}
 
-		if (InternalReplicationIndex != FNetRefHandleManager::InvalidInternalIndex && EnumHasAnyFlags(EndReplicationFlags, EEndReplicationFlags::SkipPendingEndReplicationValidation))
+		if (InternalReplicationIndex != InvalidInternalNetRefIndex && EnumHasAnyFlags(EndReplicationFlags, EEndReplicationFlags::SkipPendingEndReplicationValidation))
 		{
 			NetRefHandleManager->GetReplicatedObjectDataNoCheck(InternalReplicationIndex).bPendingEndReplication = 1U;
 		}
@@ -656,7 +661,7 @@ void UReplicationBridge::UpdateHandlesPendingEndReplication()
 	CSV_CUSTOM_STAT(Iris, NumHandlesPendingEndRepliction, (float)HandlesPendingEndReplication.Num(), ECsvCustomStatOp::Set);
 }
 
-void UReplicationBridge::AddPendingEndReplication(FNetRefHandle Handle, UE::Net::Private::FInternalNetRefIndex InternalObjectIndex, EEndReplicationFlags DestroyFlags, EPendingEndReplicationImmediate Immediate)
+void UReplicationBridge::AddPendingEndReplication(FNetRefHandle Handle, UE::Net::FInternalNetRefIndex InternalObjectIndex, EEndReplicationFlags DestroyFlags, EPendingEndReplicationImmediate Immediate)
 {
 	if (ensure(EnumHasAnyFlags(DestroyFlags, EEndReplicationFlags::Flush | EEndReplicationFlags::TearOff | EEndReplicationFlags::AllowAsyncStopReplication)))
 	{
@@ -690,7 +695,7 @@ bool UReplicationBridge::RestartOrStopObjectPendingAsyncStopReplication(FNetRefH
 		UE_LOG_REPLICATIONBRIDGE(Verbose, TEXT("RestartOrStopObjectPendingAsyncStopReplication: %s ForceStop: %u Flags:%s"), *PrintObjectFromNetRefHandle(Handle), bForceStopReplication ? 1U : 0U, *LexToString(HandlesPendingEndReplication[Index].EndReplicationFlags));
 
 		const uint32 InternalObjectIndex = NetRefHandleManager->GetInternalIndex(HandlesPendingEndReplication[Index].Handle);
-		if (InternalObjectIndex != FNetRefHandleManager::InvalidInternalIndex)
+		if (InternalObjectIndex != InvalidInternalNetRefIndex)
 		{
 			FNetRefHandleManager::FReplicatedObjectData& ObjectData = NetRefHandleManager->GetReplicatedObjectDataNoCheck(InternalObjectIndex);
 			
@@ -732,12 +737,12 @@ bool UReplicationBridge::RestartOrStopObjectPendingAsyncStopReplication(FNetRefH
 	return Index == INDEX_NONE;
 }
 
-void UReplicationBridge::InternalFlushStateData(UE::Net::FNetSerializationContext& SerializationContext, UE::Net::Private::FChangeMaskCache& ChangeMaskCache, UE::Net::FNetBitStreamWriter& ChangeMaskWriter, UE::Net::Private::FInternalNetRefIndex InternalObjectIndex)
+void UReplicationBridge::InternalFlushStateData(UE::Net::FNetSerializationContext& SerializationContext, UE::Net::Private::FChangeMaskCache& ChangeMaskCache, UE::Net::FNetBitStreamWriter& ChangeMaskWriter, UE::Net::FInternalNetRefIndex InternalObjectIndex)
 {
 	using namespace UE::Net;
 	using namespace UE::Net::Private;
 
-	if (InternalObjectIndex == FNetRefHandleManager::InvalidInternalIndex)
+	if (InternalObjectIndex == InvalidInternalNetRefIndex)
 	{
 		return;
 	}
@@ -774,14 +779,14 @@ void UReplicationBridge::InternalFlushStateData(UE::Net::FNetSerializationContex
 	}
 }
 
-void UReplicationBridge::InternalFlushStateData(UE::Net::Private::FInternalNetRefIndex InternalObjectIndex)
+void UReplicationBridge::InternalFlushStateData(UE::Net::FInternalNetRefIndex InternalObjectIndex)
 {
 	using namespace UE::Net;
 	using namespace UE::Net::Private;
 
 	IRIS_PROFILER_SCOPE(InternalFlushStateData);
 
-	if (InternalObjectIndex == FNetRefHandleManager::InvalidInternalIndex)
+	if (InternalObjectIndex == InvalidInternalNetRefIndex)
 	{
 		return;
 	}
@@ -824,7 +829,7 @@ void UReplicationBridge::InternalTearOff(FNetRefHandle Handle)
 	IRIS_PROFILER_SCOPE(InternalTearOff);
 
 	const uint32 InternalObjectIndex = NetRefHandleManager->GetInternalIndex(Handle);
-	if (InternalObjectIndex == FNetRefHandleManager::InvalidInternalIndex)
+	if (InternalObjectIndex == InvalidInternalNetRefIndex)
 	{
 		return;
 	}
@@ -874,7 +879,7 @@ void UReplicationBridge::InternalTearOff(FNetRefHandle Handle)
 		FChangeMaskCache::FCachedInfo& Info = ChangeMaskCache.AddEmptyChangeMaskForObject(InternalObjectIndex);
 		// If we are a subobject we must also mark owner as dirty.
 		const uint32 SubObjectOwnerIndex = ObjectData.SubObjectRootIndex;
-		if (SubObjectOwnerIndex != FNetRefHandleManager::InvalidInternalIndex) 
+		if (SubObjectOwnerIndex != InvalidInternalNetRefIndex) 
 		{
 			ChangeMaskCache.AddSubObjectOwnerDirty(SubObjectOwnerIndex);
 		}			
@@ -910,10 +915,11 @@ void UReplicationBridge::InternalTearOff(FNetRefHandle Handle)
 
 bool UReplicationBridge::IsReplicatedHandle(FNetRefHandle Handle) const
 {
+	using namespace UE::Net;
 	using namespace UE::Net::Private;
 
 	const FInternalNetRefIndex InternalObjectIndex = NetRefHandleManager->GetInternalIndex(Handle);
-	if (InternalObjectIndex == FNetRefHandleManager::InvalidInternalIndex)
+	if (InternalObjectIndex == InvalidInternalNetRefIndex)
 	{
 		return false;
 	}
@@ -1025,7 +1031,7 @@ const TMap<FObjectKey, UE::Net::FNetObjectGroupHandle>& UReplicationBridge::GetA
 	return ContainerGroups;
 }
 
-void UReplicationBridge::DestroyGlobalNetHandle(UE::Net::Private::FInternalNetRefIndex InternalReplicationIndex)
+void UReplicationBridge::DestroyGlobalNetHandle(UE::Net::FInternalNetRefIndex InternalReplicationIndex)
 {
 	using namespace UE::Net;
 	using namespace UE::Net::Private;
@@ -1037,7 +1043,7 @@ void UReplicationBridge::DestroyGlobalNetHandle(UE::Net::Private::FInternalNetRe
 	}
 }
 
-void UReplicationBridge::ClearNetPushIds(UE::Net::Private::FInternalNetRefIndex InternalReplicationIndex)
+void UReplicationBridge::ClearNetPushIds(UE::Net::FInternalNetRefIndex InternalReplicationIndex)
 {
 #if WITH_PUSH_MODEL
 	using namespace UE::Net;
diff --git a/Engine/Source/Runtime/Net/Iris/Private/Iris/ReplicationSystem/ReplicationReader.cpp b/Engine/Source/Runtime/Net/Iris/Private/Iris/ReplicationSystem/ReplicationReader.cpp
--- a/Engine/Source/Runtime/Net/Iris/Private/Iris/ReplicationSystem/ReplicationReader.cpp
+++ b/Engine/Source/Runtime/Net/Iris/Private/Iris/ReplicationSystem/ReplicationReader.cpp
@@ -197,7 +197,7 @@ private:
 
 struct FReplicationReader::FDispatchObjectInfo
 {
-	FInternalNetRefIndex InternalIndex = FNetRefHandleManager::InvalidInternalIndex;
+	FInternalNetRefIndex InternalIndex = InvalidInternalNetRefIndex;
 	FChangeMaskStorageOrPointer ChangeMaskOrPointer;
 	bool bIsInitialState : 1 = false;
 	bool bHasState : 1 = false;
@@ -298,7 +298,7 @@ private:
 //------------------------------------------------------------------------
 
 FReplicationReader::FReplicatedObjectInfo::FReplicatedObjectInfo()
-: InternalIndex(FNetRefHandleManager::InvalidInternalIndex)
+: InternalIndex(InvalidInternalNetRefIndex)
 , Value(0U)
 {
 	FMemory::Memzero(StoredBaselines);
@@ -521,7 +521,7 @@ uint32 FReplicationReader::ReadRootObjectsPendingDestroy(FNetSerializationContex
 		{
 			// Resolve handle and destroy using bridge
 			const FInternalNetRefIndex InternalIndex = NetRefHandleManager->GetInternalIndex(IncompleteHandle);
-			if (InternalIndex != FNetRefHandleManager::InvalidInternalIndex)
+			if (InternalIndex != InvalidInternalNetRefIndex)
 			{
 				UE_NET_TRACE_SET_SCOPE_OBJECTID(DestroyedObjectScope, IncompleteHandle);
 
@@ -1259,7 +1259,7 @@ void FReplicationReader::ReadObjectInBatch(FNetSerializationContext& Context, FN
 			const FNetRefHandle IncompleteOwnerHandle = BatchHandle;
 				
 			FInternalNetRefIndex RootObjectInternalIndex = NetRefHandleManager->GetInternalIndex(IncompleteOwnerHandle);
-			if (Reader.IsOverflown() || RootObjectInternalIndex == FNetRefHandleManager::InvalidInternalIndex)
+			if (Reader.IsOverflown() || RootObjectInternalIndex == InvalidInternalNetRefIndex)
 			{
 				UE_LOG(LogIris, Error, TEXT("FReplicationReader::ReadObject Invalid subobjectowner handle. %s"), ToCStr(IncompleteOwnerHandle.ToString()));
 				const FName& NetError = (Reader.IsOverflown() ? GNetError_BitStreamOverflow : GNetError_InvalidNetHandle);
@@ -1344,7 +1344,7 @@ void FReplicationReader::ReadObjectInBatch(FNetSerializationContext& Context, FN
 			// If we get back an invalid internal index then either the object has been deleted or there's bitstream corruption.
 			InternalIndex = NetRefHandleManager->GetInternalIndex(IncompleteHandle);
 
-			if (InternalIndex == FNetRefHandleManager::InvalidInternalIndex)
+			if (InternalIndex == InvalidInternalNetRefIndex)
 			{
 				if (!EnumHasAnyFlags(ReplicatedDestroyHeaderFlags, EReplicatedDestroyHeaderFlags::IsSubObject))
 				{
@@ -2074,7 +2074,7 @@ void FReplicationReader::DispatchStateData(FNetSerializationContext& Context)
 		const FNetRefHandleManager::FReplicatedObjectData& ObjectData = NetRefHandleManager->GetReplicatedObjectDataNoCheck(Info.InternalIndex);
 
 		// Before starting a potentially new batch we want to flush rpc's and legacy callbacks belonging to the previous batch
-		const FInternalNetRefIndex RootInternalIndex = ObjectData.SubObjectRootIndex == FNetRefHandleManager::InvalidInternalIndex ? Info.InternalIndex : ObjectData.SubObjectRootIndex;
+		const FInternalNetRefIndex RootInternalIndex = ObjectData.SubObjectRootIndex == InvalidInternalNetRefIndex ? Info.InternalIndex : ObjectData.SubObjectRootIndex;
 		if (RootInternalIndex != LastDispatchedRootInternalIndex && NumObjectsPendingPostDistpatch)
 		{
 			FlushPostDispatchForBatch();
@@ -2094,7 +2094,7 @@ void FReplicationReader::DispatchStateData(FNetSerializationContext& Context)
 		// For SubObjects we call this method after applying state data for the owner, in order to remain backwards compatible.
 		if (Info.bShouldCallSubObjectCreatedFromReplication)
 		{
-			if (ObjectData.SubObjectRootIndex != FNetRefHandleManager::InvalidInternalIndex)
+			if (ObjectData.SubObjectRootIndex != InvalidInternalNetRefIndex)
 			{
 				FBridgeSerialization::SubObjectCreatedFromReplication(ReplicationBridge, ObjectData.SubObjectRootIndex, ObjectData.RefHandle);
 			}
@@ -2681,7 +2681,7 @@ void FReplicationReader::ProcessQueuedBatches()
 
 				// End replication for object
 				const uint32 InternalIndex = NetRefHandleManager->GetInternalIndex(NetRefHandleToEndReplication);
-				if (InternalIndex != FNetRefHandleManager::InvalidInternalIndex)
+				if (InternalIndex != InvalidInternalNetRefIndex)
 				{
 					constexpr bool bTearOff = false;
 					EndReplication(InternalIndex, bTearOff, bShouldDestroyInstance, CloseReason, DetachReason);
@@ -3215,7 +3215,7 @@ FString FReplicationReader::PrintObjectInfo(FInternalNetRefIndex ObjectIndex, FN
 	TStringBuilder<512> InfoBuilder;
 
 	const FNetRefHandleManager::FReplicatedObjectData& ObjectData = NetRefHandleManager->GetReplicatedObjectDataNoCheck(ObjectIndex);
-	const bool bIsSubObject = ObjectData.SubObjectRootIndex != FNetRefHandleManager::InvalidInternalIndex;
+	const bool bIsSubObject = ObjectData.SubObjectRootIndex != InvalidInternalNetRefIndex;
 	InfoBuilder.Appendf(TEXT("Status info: 0x%llx (%s)"), ObjectInfo->Value, bIsSubObject?TEXT("SubObject"):TEXT("RootObject"));
 
 	const FInternalNetRefIndex RootInternalIndex = bIsSubObject ? ObjectData.SubObjectRootIndex : ObjectIndex;
diff --git a/Engine/Source/Runtime/Net/Iris/Private/Iris/ReplicationSystem/ReplicationSystem.cpp b/Engine/Source/Runtime/Net/Iris/Private/Iris/ReplicationSystem/ReplicationSystem.cpp
--- a/Engine/Source/Runtime/Net/Iris/Private/Iris/ReplicationSystem/ReplicationSystem.cpp
+++ b/Engine/Source/Runtime/Net/Iris/Private/Iris/ReplicationSystem/ReplicationSystem.cpp
@@ -1357,8 +1357,8 @@ void UReplicationSystem::SetReplicationView(uint32 ConnectionId, const UE::Net::
 
 void UReplicationSystem::SetStaticPriority(FNetRefHandle Handle, float Priority)
 {
-	const UE::Net::Private::FInternalNetRefIndex ObjectInternalIndex = Impl->ReplicationSystemInternal.GetNetRefHandleManager().GetInternalIndex(Handle);
-	if (ObjectInternalIndex == UE::Net::Private::FNetRefHandleManager::InvalidInternalIndex)
+	const UE::Net::FInternalNetRefIndex ObjectInternalIndex = Impl->ReplicationSystemInternal.GetNetRefHandleManager().GetInternalIndex(Handle);
+	if (ObjectInternalIndex == UE::Net::InvalidInternalNetRefIndex)
 	{
 		return;
 	}
@@ -1368,6 +1368,7 @@ void UReplicationSystem::SetStaticPriority(FNetRefHandle Handle, float Priority)
 
 bool UReplicationSystem::SetPrioritizer(FNetRefHandle Handle, UE::Net::FNetObjectPrioritizerHandle Prioritizer)
 {
+	using namespace UE::Net;
 	using namespace UE::Net::Private;
 
 	if (!Handle.IsValid())
@@ -1377,7 +1378,7 @@ bool UReplicationSystem::SetPrioritizer(FNetRefHandle Handle, UE::Net::FNetObjec
 
 	FReplicationSystemInternal& ReplicationSystemInternal = Impl->ReplicationSystemInternal;
 	const FInternalNetRefIndex ObjectInternalIndex = ReplicationSystemInternal.GetNetRefHandleManager().GetInternalIndex(Handle);
-	if (ObjectInternalIndex == FNetRefHandleManager::InvalidInternalIndex)
+	if (ObjectInternalIndex == UE::Net::InvalidInternalNetRefIndex)
 	{
 		return false;
 	}
@@ -1506,7 +1507,7 @@ bool UReplicationSystem::IsRelevantForConnection(FNetRefHandle Handle, uint32 Co
 
 	const FNetRefHandleManager& NetRefHandleManager = Impl->ReplicationSystemInternal.GetNetRefHandleManager();
 	const FInternalNetRefIndex ObjectInternalIndex = NetRefHandleManager.GetInternalIndex(Handle);
-	if (ObjectInternalIndex == FNetRefHandleManager::InvalidInternalIndex)
+	if (ObjectInternalIndex == UE::Net::InvalidInternalNetRefIndex)
 	{
 		return false;
 	}
@@ -1644,12 +1645,13 @@ bool UReplicationSystem::IsNetRefHandleAssigned(FNetRefHandle Handle) const
 
 const UE::Net::FReplicationProtocol* UReplicationSystem::GetReplicationProtocol(FNetRefHandle Handle) const
 {
+	using namespace UE::Net;
 	using namespace UE::Net::Private;
 
 	FNetRefHandleManager& NetRefHandleManager = Impl->ReplicationSystemInternal.GetNetRefHandleManager();
 
 	const FInternalNetRefIndex ObjectInternalIndex = NetRefHandleManager.GetInternalIndex(Handle);
-	if (ObjectInternalIndex == FNetRefHandleManager::InvalidInternalIndex)
+	if (ObjectInternalIndex == UE::Net::InvalidInternalNetRefIndex)
 	{
 		return nullptr;
 	}
@@ -1677,6 +1679,7 @@ void UReplicationSystem::DestroyFilter(const FName Name)
 
 void UReplicationSystem::SetOwningNetConnection(FNetRefHandle Handle, uint32 ConnectionId)
 {
+	using namespace UE::Net;
 	using namespace UE::Net::Private;
 
 	FNetRefHandleManager& NetRefHandleManager = Impl->ReplicationSystemInternal.GetNetRefHandleManager();
@@ -1688,7 +1691,7 @@ void UReplicationSystem::SetOwningNetConnection(FNetRefHandle Handle, uint32 Con
 		return;
 	}
 
-	if (ObjectInternalIndex == FNetRefHandleManager::InvalidInternalIndex)
+	if (ObjectInternalIndex == UE::Net::InvalidInternalNetRefIndex)
 	{
 		return;
 	}
@@ -1702,11 +1705,12 @@ void UReplicationSystem::SetOwningNetConnection(FNetRefHandle Handle, uint32 Con
 
 uint32 UReplicationSystem::GetOwningNetConnection(FNetRefHandle Handle) const
 {
+	using namespace UE::Net;
 	using namespace UE::Net::Private;
 
 	FNetRefHandleManager& NetRefHandleManager = Impl->ReplicationSystemInternal.GetNetRefHandleManager();
 	const FInternalNetRefIndex ObjectInternalIndex = NetRefHandleManager.GetInternalIndex(Handle);
-	if (ObjectInternalIndex == FNetRefHandleManager::InvalidInternalIndex)
+	if (ObjectInternalIndex == UE::Net::InvalidInternalNetRefIndex)
 	{
 		return UE::Net::InvalidConnectionId;
 	}
@@ -1717,11 +1721,12 @@ uint32 UReplicationSystem::GetOwningNetConnection(FNetRefHandle Handle) const
 
 bool UReplicationSystem::SetFilter(FNetRefHandle Handle, UE::Net::FNetObjectFilterHandle Filter, FName FilterConfigProfile /*= NAME_None*/)
 {
+	using namespace UE::Net;
 	using namespace UE::Net::Private;
 
 	FNetRefHandleManager& NetRefHandleManager = Impl->ReplicationSystemInternal.GetNetRefHandleManager();
 	const FInternalNetRefIndex ObjectInternalIndex = NetRefHandleManager.GetInternalIndex(Handle);
-	if (ObjectInternalIndex == FNetRefHandleManager::InvalidInternalIndex)
+	if (ObjectInternalIndex == UE::Net::InvalidInternalNetRefIndex)
 	{
 		return false;
 	}
@@ -1862,6 +1867,9 @@ UE::Net::FNetObjectGroupHandle UReplicationSystem::CreateGroup(FName GroupName)
 
 void UReplicationSystem::AddToGroup(FNetObjectGroupHandle GroupHandle, FNetRefHandle Handle)
 {
+	using namespace UE::Net;
+	using namespace UE::Net::Private;
+
 	// Early out if this is invalid group
 	if (!ensure(IsValidGroup(GroupHandle)))
 	{
@@ -1870,8 +1878,6 @@ void UReplicationSystem::AddToGroup(FNetObjectGroupHandle GroupHandle, FNetRefHa
 
 	LLM_SCOPE_BYTAG(Iris);
 
-	using namespace UE::Net::Private;
-
 	FNetRefHandleManager& NetRefHandleManager = Impl->ReplicationSystemInternal.GetNetRefHandleManager();
 	const FInternalNetRefIndex ObjectInternalIndex = NetRefHandleManager.GetInternalIndex(Handle);
 
@@ -1881,7 +1887,7 @@ void UReplicationSystem::AddToGroup(FNetObjectGroupHandle GroupHandle, FNetRefHa
 		return;
 	}
 
-	if (ObjectInternalIndex == FNetRefHandleManager::InvalidInternalIndex)
+	if (ObjectInternalIndex == UE::Net::InvalidInternalNetRefIndex)
 	{
 		return;
 	}
@@ -1895,14 +1901,15 @@ void UReplicationSystem::AddToGroup(FNetObjectGroupHandle GroupHandle, FNetRefHa
 
 void UReplicationSystem::RemoveFromGroup(FNetObjectGroupHandle GroupHandle, FNetRefHandle Handle)
 {
+	using namespace UE::Net;
+	using namespace UE::Net::Private;
+
 	// Early out if this is invalid group
 	if (!ensure(IsValidGroup(GroupHandle)))
 	{
 		return;
 	}
 
-	using namespace UE::Net::Private;
-
 	FNetRefHandleManager& NetRefHandleManager = Impl->ReplicationSystemInternal.GetNetRefHandleManager();
 	const FInternalNetRefIndex ObjectInternalIndex = NetRefHandleManager.GetInternalIndex(Handle);
 	if (Impl->ReplicationSystemInternal.AreFilterChangesBlocked())
@@ -1911,7 +1918,7 @@ void UReplicationSystem::RemoveFromGroup(FNetObjectGroupHandle GroupHandle, FNet
 		return;
 	}
 
-	if (ObjectInternalIndex == FNetRefHandleManager::InvalidInternalIndex)
+	if (ObjectInternalIndex == UE::Net::InvalidInternalNetRefIndex)
 	{
 		return;
 	}
@@ -1925,6 +1932,7 @@ void UReplicationSystem::RemoveFromGroup(FNetObjectGroupHandle GroupHandle, FNet
 
 void UReplicationSystem::RemoveFromAllGroups(FNetRefHandle Handle)
 {
+	using namespace UE::Net;
 	using namespace UE::Net::Private;
 
 	FNetRefHandleManager& NetRefHandleManager = Impl->ReplicationSystemInternal.GetNetRefHandleManager();
@@ -1935,14 +1943,15 @@ void UReplicationSystem::RemoveFromAllGroups(FNetRefHandle Handle)
 
 bool UReplicationSystem::IsInGroup(FNetObjectGroupHandle GroupHandle, FNetRefHandle Handle) const
 {
+	using namespace UE::Net;
+	using namespace UE::Net::Private;
+
 	// Early out if this is invalid group
 	if (!IsValidGroup(GroupHandle))
 	{
 		return false;
 	}
 
-	using namespace UE::Net::Private;
-
 	const FNetObjectGroups& Groups = Impl->ReplicationSystemInternal.GetGroups();
 	const FNetRefHandleManager& NetRefHandleManager = Impl->ReplicationSystemInternal.GetNetRefHandleManager();
 
@@ -2120,11 +2129,12 @@ void UReplicationSystem::SetGroupFilterStatus(FNetObjectGroupHandle GroupHandle,
 
 bool UReplicationSystem::SetReplicationConditionConnectionFilter(FNetRefHandle Handle, UE::Net::EReplicationCondition Condition, uint32 ConnectionId, bool bEnable)
 {
+	using namespace UE::Net;
 	using namespace UE::Net::Private;
 
 	FNetRefHandleManager& NetRefHandleManager = Impl->ReplicationSystemInternal.GetNetRefHandleManager();
 	const FInternalNetRefIndex ObjectInternalIndex = NetRefHandleManager.GetInternalIndex(Handle);
-	if (ObjectInternalIndex == FNetRefHandleManager::InvalidInternalIndex)
+	if (ObjectInternalIndex == UE::Net::InvalidInternalNetRefIndex)
 	{
 		return false;
 	}
@@ -2141,11 +2151,12 @@ bool UReplicationSystem::SetReplicationConditionConnectionFilter(FNetRefHandle H
 
 bool UReplicationSystem::SetReplicationCondition(FNetRefHandle Handle, UE::Net::EReplicationCondition Condition, bool bEnable)
 {
+	using namespace UE::Net;
 	using namespace UE::Net::Private;
 
 	FNetRefHandleManager& NetRefHandleManager = Impl->ReplicationSystemInternal.GetNetRefHandleManager();
 	const FInternalNetRefIndex ObjectInternalIndex = NetRefHandleManager.GetInternalIndex(Handle);
-	if (ObjectInternalIndex == FNetRefHandleManager::InvalidInternalIndex)
+	if (ObjectInternalIndex == UE::Net::InvalidInternalNetRefIndex)
 	{
 		return false;
 	}
@@ -2156,12 +2167,13 @@ bool UReplicationSystem::SetReplicationCondition(FNetRefHandle Handle, UE::Net::
 
 void UReplicationSystem::SetDeltaCompressionStatus(FNetRefHandle Handle, UE::Net::ENetObjectDeltaCompressionStatus Status)
 {
+	using namespace UE::Net;
 	using namespace UE::Net::Private;
 
 	FReplicationSystemInternal& InternalSys = Impl->ReplicationSystemInternal;
 	FNetRefHandleManager& NetRefHandleManager = InternalSys.GetNetRefHandleManager();
 	const FInternalNetRefIndex ObjectInternalIndex = NetRefHandleManager.GetInternalIndex(Handle);
-	if (ObjectInternalIndex == FNetRefHandleManager::InvalidInternalIndex)
+	if (ObjectInternalIndex == UE::Net::InvalidInternalNetRefIndex)
 	{
 		return;
 	}
@@ -2244,8 +2256,8 @@ const UE::Net::FWorldLocations& UReplicationSystem::GetWorldLocations() const
 
 void UReplicationSystem::SetCullDistanceOverride(FNetRefHandle Handle, float CullDistance)
 {
-	const UE::Net::Private::FInternalNetRefIndex ObjectInternalIndex = Impl->ReplicationSystemInternal.GetNetRefHandleManager().GetInternalIndex(Handle);
-	if (ObjectInternalIndex == UE::Net::Private::FNetRefHandleManager::InvalidInternalIndex)
+	const UE::Net::FInternalNetRefIndex ObjectInternalIndex = Impl->ReplicationSystemInternal.GetNetRefHandleManager().GetInternalIndex(Handle);
+	if (ObjectInternalIndex == UE::Net::InvalidInternalNetRefIndex)
 	{
 		return;
 	}
@@ -2263,8 +2275,8 @@ void UReplicationSystem::SetCullDistanceOverride(FNetRefHandle Handle, float Cul
 
 void UReplicationSystem::ClearCullDistanceOverride(FNetRefHandle Handle)
 {
-	const UE::Net::Private::FInternalNetRefIndex ObjectInternalIndex = Impl->ReplicationSystemInternal.GetNetRefHandleManager().GetInternalIndex(Handle);
-	if (ObjectInternalIndex == UE::Net::Private::FNetRefHandleManager::InvalidInternalIndex)
+	const UE::Net::FInternalNetRefIndex ObjectInternalIndex = Impl->ReplicationSystemInternal.GetNetRefHandleManager().GetInternalIndex(Handle);
+	if (ObjectInternalIndex == UE::Net::InvalidInternalNetRefIndex)
 	{
 		return;
 	}
@@ -2278,8 +2290,8 @@ void UReplicationSystem::ClearCullDistanceOverride(FNetRefHandle Handle)
 
 float UReplicationSystem::GetCullDistance(FNetRefHandle Handle, float DefaultValue) const
 {
-	const UE::Net::Private::FInternalNetRefIndex ObjectInternalIndex = Impl->ReplicationSystemInternal.GetNetRefHandleManager().GetInternalIndex(Handle);
-	if (ObjectInternalIndex == UE::Net::Private::FNetRefHandleManager::InvalidInternalIndex)
+	const UE::Net::FInternalNetRefIndex ObjectInternalIndex = Impl->ReplicationSystemInternal.GetNetRefHandleManager().GetInternalIndex(Handle);
+	if (ObjectInternalIndex == UE::Net::InvalidInternalNetRefIndex)
 	{
 		return DefaultValue;
 	}
diff --git a/Engine/Source/Runtime/Net/Iris/Private/Iris/ReplicationSystem/ReplicationSystemInternal.h b/Engine/Source/Runtime/Net/Iris/Private/Iris/ReplicationSystem/ReplicationSystemInternal.h
--- a/Engine/Source/Runtime/Net/Iris/Private/Iris/ReplicationSystem/ReplicationSystemInternal.h
+++ b/Engine/Source/Runtime/Net/Iris/Private/Iris/ReplicationSystem/ReplicationSystemInternal.h
@@ -205,7 +205,7 @@ public:
 	}
 #endif // UE_SUPPORT_PARALLEL_IRIS
 
-void RemoveFromAllGroups(UE::Net::Private::FInternalNetRefIndex ObjectInternalIndex);
+void RemoveFromAllGroups(UE::Net::FInternalNetRefIndex ObjectInternalIndex);
 
 private:
 	FReplicationProtocolManager ReplicationProtocolManager;
diff --git a/Engine/Source/Runtime/Net/Iris/Private/Iris/ReplicationSystem/ReplicationTypes.h b/Engine/Source/Runtime/Net/Iris/Private/Iris/ReplicationSystem/ReplicationTypes.h
--- a/Engine/Source/Runtime/Net/Iris/Private/Iris/ReplicationSystem/ReplicationTypes.h
+++ b/Engine/Source/Runtime/Net/Iris/Private/Iris/ReplicationSystem/ReplicationTypes.h
@@ -10,7 +10,7 @@
 
 class UReplicationSystem;
 
-namespace UE::Net::Private
+namespace UE::Net
 {
 	typedef uint32 FInternalNetRefIndex;
 }
diff --git a/Engine/Source/Runtime/Net/Iris/Private/Iris/ReplicationSystem/ReplicationWriter.cpp b/Engine/Source/Runtime/Net/Iris/Private/Iris/ReplicationSystem/ReplicationWriter.cpp
--- a/Engine/Source/Runtime/Net/Iris/Private/Iris/ReplicationSystem/ReplicationWriter.cpp
+++ b/Engine/Source/Runtime/Net/Iris/Private/Iris/ReplicationSystem/ReplicationWriter.cpp
@@ -291,7 +291,7 @@ bool FReplicationWriter::IsInReplicatingState(FInternalNetRefIndex ObjectInterna
 		return false;
 	}
 
-	// We let InvalidInternalIndex be accessed because it has an object state too.
+	// We let InvalidInternalNetRefIndex be accessed because it has an object state too.
 	const FReplicationInfo& Info = GetReplicationInfo(ObjectInternalIndex);
 	const EReplicatedObjectState CurrentState = Info.GetState();
 	
@@ -309,7 +309,7 @@ bool FReplicationWriter::QueueNetObjectAttachments(FInternalNetRefIndex OwnerInt
 		return false;
 	}
 
-	const uint32 TargetIndex = SubObjectInternalIndex != FNetRefHandleManager::InvalidInternalIndex ? SubObjectInternalIndex : OwnerInternalIndex;
+	const uint32 TargetIndex = SubObjectInternalIndex != InvalidInternalNetRefIndex ? SubObjectInternalIndex : OwnerInternalIndex;
 	const bool bTargetObjectInScope = ObjectsInScope.GetBit(TargetIndex);
 	if (!bTargetObjectInScope && !Parameters.bAllowSendingAttachmentsToObjectsNotInScope)
 	{
@@ -624,7 +624,7 @@ void FReplicationWriter::StartReplication(uint32 InternalIndex)
 	if (Info.IsSubObject)
 	{
 		const uint32 RootObjectInternalIndex = NetRefHandleManager->GetRootObjectInternalIndexOfSubObject(InternalIndex);
-		if (ensure(RootObjectInternalIndex != FNetRefHandleManager::InvalidInternalIndex))
+		if (ensure(RootObjectInternalIndex != InvalidInternalNetRefIndex))
 		{
 			FReplicationInfo& OwnerInfo = ReplicatedObjects[RootObjectInternalIndex];
 			if (ensureMsgf(OwnerInfo.GetState() != EReplicatedObjectState::Invalid && OwnerInfo.GetState() < EReplicatedObjectState::PendingDestroy, 
@@ -1822,7 +1822,7 @@ void FReplicationWriter::HandleDroppedRecord<FReplicationWriter::EReplicatedObje
 			if (Info.IsSubObject)
 			{
 				const FNetRefHandleManager::FReplicatedObjectData& ObjectData = NetRefHandleManager->GetReplicatedObjectDataNoCheck(InternalIndex);
-				if (ensure(ObjectData.SubObjectRootIndex != FNetRefHandleManager::InvalidInternalIndex))
+				if (ensure(ObjectData.SubObjectRootIndex != InvalidInternalNetRefIndex))
 				{
 					ParentInternalIndex = ObjectData.SubObjectRootIndex;
 					MarkObjectDirty(ParentInternalIndex, "DroppedWaitOnDestroy2");
@@ -2579,7 +2579,7 @@ FReplicationWriter::EWriteObjectStatus FReplicationWriter::WriteObjectAndSubObje
 		{
 			const FHugeObjectSendQueue& HugeObjectSendQueueForDebug = OOBWriteAttachmentType == ENetObjectAttachmentType::HugeObject ? HugeObjectSendQueue : DebugObjectSendQueue;
 			const FInternalNetRefIndex HugeObjectInternalIndex = HugeObjectSendQueueForDebug.GetRootObjectInternalIndexForTrace();
-			if (HugeObjectInternalIndex != FNetRefHandleManager::InvalidInternalIndex)
+			if (HugeObjectInternalIndex != InvalidInternalNetRefIndex)
 			{
 				NetRefHandleForTraceScope = NetRefHandleManager->GetReplicatedObjectDataNoCheck(HugeObjectInternalIndex).RefHandle;
 			}
@@ -3408,7 +3408,7 @@ FReplicationWriter::EWriteStatus FReplicationWriter::PrepareAndSendHugeObjectPay
 		
 		FBatchInfo HugeObjectBatchInfo;
 		HugeObjectBatchInfo.Type = EBatchInfoType::HugeObject;
-		HugeObjectBatchInfo.ParentInternalIndex = FNetRefHandleManager::InvalidInternalIndex;
+		HugeObjectBatchInfo.ParentInternalIndex = InvalidInternalNetRefIndex;
 
 		const FWriteObjectTraits WriteHugeObjectInfo({.bWriteAttachments = 1U, .OOBWriteAttachmentType = (uint32)ENetObjectAttachmentType::HugeObject});
 		const EWriteObjectStatus HugeObjectStatus = WriteObjectInBatch(Context, ObjectIndexForOOBAttachment, WriteHugeObjectInfo, HugeObjectBatchInfo);
@@ -4896,7 +4896,7 @@ bool FReplicationWriter::FHugeObjectSendQueue::IsObjectInQueue(FInternalNetRefIn
 		 return TailNode->GetValue().RootObjectInternalIndex;
 	 }
 
-	return FNetRefHandleManager::InvalidInternalIndex;
+	return InvalidInternalNetRefIndex;
  }
 
 void FReplicationWriter::FHugeObjectSendQueue::AckObjects(TFunctionRef<void (const FHugeObjectContext& Context)> AckHugeObject)
diff --git a/Engine/Source/Runtime/Net/Iris/Private/Iris/ReplicationSystem/WorldLocations.cpp b/Engine/Source/Runtime/Net/Iris/Private/Iris/ReplicationSystem/WorldLocations.cpp
--- a/Engine/Source/Runtime/Net/Iris/Private/Iris/ReplicationSystem/WorldLocations.cpp
+++ b/Engine/Source/Runtime/Net/Iris/Private/Iris/ReplicationSystem/WorldLocations.cpp
@@ -39,7 +39,7 @@ void FWorldLocations::PostSendUpdate()
 	ObjectsWithDirtyInfo.ClearAllBits();
 }
 
-void FWorldLocations::OnMaxInternalNetRefIndexIncreased(UE::Net::Private::FInternalNetRefIndex NewMaxInternalIndex)
+void FWorldLocations::OnMaxInternalNetRefIndexIncreased(FInternalNetRefIndex NewMaxInternalIndex)
 {
 	ValidInfoIndexes.SetNumBits(NewMaxInternalIndex);
 	ObjectsWithDirtyInfo.SetNumBits(NewMaxInternalIndex);
@@ -47,7 +47,7 @@ void FWorldLocations::OnMaxInternalNetRefIndexIncreased(UE::Net::Private::FInter
 	StorageIndexes.SetNum(NewMaxInternalIndex);
 }
 
-void FWorldLocations::InitObjectInfoCache(UE::Net::Private::FInternalNetRefIndex ObjectIndex)
+void FWorldLocations::InitObjectInfoCache(FInternalNetRefIndex ObjectIndex)
 {
 	if (ValidInfoIndexes.IsBitSet(ObjectIndex))
 	{
@@ -73,7 +73,7 @@ void FWorldLocations::InitObjectInfoCache(UE::Net::Private::FInternalNetRefIndex
 	StorageIndexes[ObjectIndex] = AvailableSlot;
 }
 
-void FWorldLocations::RemoveObjectInfoCache(UE::Net::Private::FInternalNetRefIndex ObjectIndex)
+void FWorldLocations::RemoveObjectInfoCache(FInternalNetRefIndex ObjectIndex)
 {
 	if (!ValidInfoIndexes.IsBitSet(ObjectIndex))
 	{
@@ -94,7 +94,7 @@ void FWorldLocations::RemoveObjectInfoCache(UE::Net::Private::FInternalNetRefInd
 	ReservedStorageSlot.ClearBit(StorageIndex);
 }
 
-void FWorldLocations::SetObjectInfo(UE::Net::Private::FInternalNetRefIndex ObjectIndex, const FVector& Location, const float NetCullDistance)
+void FWorldLocations::SetObjectInfo(FInternalNetRefIndex ObjectIndex, const FVector& Location, const float NetCullDistance)
 {
 	ensure(!bLockdownDirtyList);
 
@@ -120,7 +120,7 @@ void FWorldLocations::SetObjectInfo(UE::Net::Private::FInternalNetRefIndex Objec
 	ObjectsWithDirtyInfo.OrBitValue(ObjectIndex, bHasInfoChanged);
 }
 
-bool FWorldLocations::SetCullDistanceOverride(UE::Net::Private::FInternalNetRefIndex ObjectIndex, float CullDistance)
+bool FWorldLocations::SetCullDistanceOverride(FInternalNetRefIndex ObjectIndex, float CullDistance)
 {
 	ensure(!bLockdownDirtyList);
 
@@ -145,7 +145,7 @@ bool FWorldLocations::SetCullDistanceOverride(UE::Net::Private::FInternalNetRefI
 	return false;
 }
 
-bool FWorldLocations::ClearCullDistanceOverride(UE::Net::Private::FInternalNetRefIndex ObjectIndex)
+bool FWorldLocations::ClearCullDistanceOverride(FInternalNetRefIndex ObjectIndex)
 {
 	ensure(!bLockdownDirtyList);
 
diff --git a/Engine/Source/Runtime/Net/Iris/Public/Iris/Core/IrisDebugging.h b/Engine/Source/Runtime/Net/Iris/Public/Iris/Core/IrisDebugging.h
--- a/Engine/Source/Runtime/Net/Iris/Public/Iris/Core/IrisDebugging.h
+++ b/Engine/Source/Runtime/Net/Iris/Public/Iris/Core/IrisDebugging.h
@@ -13,10 +13,7 @@ namespace UE::Net
 	struct FReplicationInstanceProtocol;
 
 	typedef uint32 FReplicationProtocolIdentifier;
-}
 
-namespace UE::Net::Private
-{
 	typedef uint32 FInternalNetRefIndex;
 }
 
@@ -38,13 +35,13 @@ IRISCORE_API bool BreakOnNetRefHandle(FNetRefHandle NetRefHandle);
 IRISCORE_API bool BreakOnRPCName(FName RPCName);
 
 /** Trigger a breakpoint and return true if the index is the current debug index */
-IRISCORE_API bool BreakOnInternalNetRefIndex(UE::Net::Private::FInternalNetRefIndex InternalIndex);
+IRISCORE_API bool BreakOnInternalNetRefIndex(UE::Net::FInternalNetRefIndex InternalIndex);
 
 /** Returns true if the object name contains the current debug name, will return true if no debug name is set */
 IRISCORE_API bool FilterDebuggedObject(UObject* Object);
 
 /** Returns the internal index set via SetIrisDebugInternalIndex */
-IRISCORE_API UE::Net::Private::FInternalNetRefIndex GetDebugNetInternalIndex();
+IRISCORE_API UE::Net::FInternalNetRefIndex GetDebugNetInternalIndex();
 
 /** Returns the NetRefHandle set via SetIrisDebugNetRefHandle */
 IRISCORE_API FNetRefHandle GetDebugNetRefHandle();
@@ -77,7 +74,7 @@ extern "C" IRISCORE_API void DebugOutputNetObjectProtocolReferences(FReplication
 struct FNetReplicatedObjectDebugInfo
 {
 	const FNetRefHandle* RefHandle;
-	UE::Net::Private::FInternalNetRefIndex InternalNetRefIndex;
+	UE::Net::FInternalNetRefIndex InternalNetRefIndex;
 	const UReplicationSystem* ReplicationSystem;
 	const FReplicationProtocol* Protocol;
 	const FReplicationInstanceProtocol* InstanceProtocol;
@@ -97,7 +94,7 @@ extern "C" IRISCORE_API FNetReplicatedObjectDebugInfo DebugNetRefHandle(FNetRefH
 extern "C" IRISCORE_API FNetReplicatedObjectDebugInfo DebugNetRefHandleById(uint64 NetRefHandleId, uint32 ReplicationSystemId);
 
 /** Look up replicated handle specified by an internal index and replicationsystem id and return debug information */
-extern "C" IRISCORE_API FNetReplicatedObjectDebugInfo DebugInternalNetRefIndex(UE::Net::Private::FInternalNetRefIndex InternalIndex, uint32 ReplicationSystemId);
+extern "C" IRISCORE_API FNetReplicatedObjectDebugInfo DebugInternalNetRefIndex(UE::Net::FInternalNetRefIndex InternalIndex, uint32 ReplicationSystemId);
 
 /** Variant of DebugOutputNetObjectProtocolReferences that can be used from breakpoints or in watch window to find all handles references registered for a protocol and output to DebugOutput in debugger
  * NOTE: Use only for debugging as this variant uses a static buffer which is not thread safe	
@@ -111,7 +108,7 @@ extern "C" IRISCORE_API void SetIrisDebugObjectName(const ANSICHAR* NameBuffer);
 extern "C" IRISCORE_API void SetIrisDebugNetRefHandle(uint64 NetHandleId);
 
 /** Set the InternalIndex to break on */
-extern "C" IRISCORE_API void SetIrisDebugInternalNetRefIndex(UE::Net::Private::FInternalNetRefIndex InternalIndex);
+extern "C" IRISCORE_API void SetIrisDebugInternalNetRefIndex(UE::Net::FInternalNetRefIndex InternalIndex);
 
 /** Set the InternalIndex to break on via it's NetHandle */
 extern "C" IRISCORE_API void SetIrisDebugInternalNetRefIndexViaNetHandle(FNetRefHandle RefHandle);
diff --git a/Engine/Source/Runtime/Net/Iris/Public/Iris/ReplicationSystem/Filtering/AlwaysRelevantNetObjectFilter.h b/Engine/Source/Runtime/Net/Iris/Public/Iris/ReplicationSystem/Filtering/AlwaysRelevantNetObjectFilter.h
--- a/Engine/Source/Runtime/Net/Iris/Public/Iris/ReplicationSystem/Filtering/AlwaysRelevantNetObjectFilter.h
+++ b/Engine/Source/Runtime/Net/Iris/Public/Iris/ReplicationSystem/Filtering/AlwaysRelevantNetObjectFilter.h
@@ -20,9 +20,9 @@ protected:
 	// UNetObjectFilter interface
 	IRISCORE_API virtual void OnInit(const FNetObjectFilterInitParams&) override;
 	IRISCORE_API virtual void OnDeinit() override;
-	IRISCORE_API virtual bool AddObject(uint32 ObjectIndex, FNetObjectFilterAddObjectParams&) override;
-	IRISCORE_API virtual void RemoveObject(uint32 ObjectIndex, const FNetObjectFilteringInfo&) override;
+	IRISCORE_API virtual bool AddObject(UE::Net::FInternalNetRefIndex ObjectIndex, FNetObjectFilterAddObjectParams&) override;
+	IRISCORE_API virtual void RemoveObject(UE::Net::FInternalNetRefIndex ObjectIndex, const FNetObjectFilteringInfo&) override;
 	IRISCORE_API virtual void UpdateObjects(FNetObjectFilterUpdateParams&) override;
 	IRISCORE_API virtual void Filter(FNetObjectFilteringParams&) override;
-	IRISCORE_API virtual void OnMaxInternalNetRefIndexIncreased(uint32 NewMaxInternalIndex) override;
+	IRISCORE_API virtual void OnMaxInternalNetRefIndexIncreased(UE::Net::FInternalNetRefIndex NewMaxInternalIndex) override;
 };
diff --git a/Engine/Source/Runtime/Net/Iris/Public/Iris/ReplicationSystem/Filtering/FilterOutNetObjectFilter.h b/Engine/Source/Runtime/Net/Iris/Public/Iris/ReplicationSystem/Filtering/FilterOutNetObjectFilter.h
--- a/Engine/Source/Runtime/Net/Iris/Public/Iris/ReplicationSystem/Filtering/FilterOutNetObjectFilter.h
+++ b/Engine/Source/Runtime/Net/Iris/Public/Iris/ReplicationSystem/Filtering/FilterOutNetObjectFilter.h
@@ -25,9 +25,9 @@ protected:
 	// UNetObjectFilter interface
 	virtual void OnInit(const FNetObjectFilterInitParams&) override {}
 	virtual void OnDeinit() override {}
-	virtual void OnMaxInternalNetRefIndexIncreased(uint32 NewMaxInternalIndex) override {}
-	virtual bool AddObject(uint32 ObjectIndex, FNetObjectFilterAddObjectParams&) override { return true; }
-	virtual void RemoveObject(uint32 ObjectIndex, const FNetObjectFilteringInfo&) override {}
+	virtual void OnMaxInternalNetRefIndexIncreased(UE::Net::FInternalNetRefIndex NewMaxInternalIndex) override {}
+	virtual bool AddObject(UE::Net::FInternalNetRefIndex ObjectIndex, FNetObjectFilterAddObjectParams&) override { return true; }
+	virtual void RemoveObject(UE::Net::FInternalNetRefIndex ObjectIndex, const FNetObjectFilteringInfo&) override {}
 
 	IRISCORE_API virtual void Filter(FNetObjectFilteringParams&) override;
 };
diff --git a/Engine/Source/Runtime/Net/Iris/Public/Iris/ReplicationSystem/Filtering/NetObjectConnectionFilter.h b/Engine/Source/Runtime/Net/Iris/Public/Iris/ReplicationSystem/Filtering/NetObjectConnectionFilter.h
--- a/Engine/Source/Runtime/Net/Iris/Public/Iris/ReplicationSystem/Filtering/NetObjectConnectionFilter.h
+++ b/Engine/Source/Runtime/Net/Iris/Public/Iris/ReplicationSystem/Filtering/NetObjectConnectionFilter.h
@@ -55,16 +55,16 @@ protected:
 	virtual void OnMaxInternalNetRefIndexIncreased(uint32 NewMaxInternalIndex) override {}
 	IRISCORE_API virtual void AddConnection(uint32 ConnectionId) override;
 	IRISCORE_API virtual void RemoveConnection(uint32 ConnectionId) override;
-	IRISCORE_API virtual bool AddObject(uint32 ObjectIndex, FNetObjectFilterAddObjectParams&) override;
-	IRISCORE_API virtual void RemoveObject(uint32 ObjectIndex, const FNetObjectFilteringInfo&) override;
+	IRISCORE_API virtual bool AddObject(UE::Net::FInternalNetRefIndex ObjectIndex, FNetObjectFilterAddObjectParams&) override;
+	IRISCORE_API virtual void RemoveObject(UE::Net::FInternalNetRefIndex ObjectIndex, const FNetObjectFilteringInfo&) override;
 	IRISCORE_API virtual void PreFilter(FNetObjectPreFilteringParams&) override;
 	IRISCORE_API virtual void Filter(FNetObjectFilteringParams&) override;
 
 	TStrongObjectPtr<UNetObjectConnectionFilterConfig> Config;
-	TArray<uint32> LocalToNetRefIndex;
+	TArray<UE::Net::FInternalNetRefIndex> LocalToNetRefIndex;
 	TArray<FPerConnectionInfo> PerConnectionInfos;
 	UE::Net::FNetBitArray UsedLocalInfoIndices;
 
-	uint32 MaxInternalIndex = 0;
+	UE::Net::FInternalNetRefIndex MaxInternalIndex = 0;
 	bool bObjectRemoved = false;
 };
diff --git a/Engine/Source/Runtime/Net/Iris/Public/Iris/ReplicationSystem/Filtering/NetObjectFilter.h b/Engine/Source/Runtime/Net/Iris/Public/Iris/ReplicationSystem/Filtering/NetObjectFilter.h
--- a/Engine/Source/Runtime/Net/Iris/Public/Iris/ReplicationSystem/Filtering/NetObjectFilter.h
+++ b/Engine/Source/Runtime/Net/Iris/Public/Iris/ReplicationSystem/Filtering/NetObjectFilter.h
@@ -9,17 +9,22 @@
 #include "NetObjectFilter.generated.h"
 
 struct FNetObjectFilteringInfo;
+
 class UReplicationSystem;
+
 namespace UE::Net
 {
 	typedef uint32 FNetObjectFilterHandle;
+	typedef uint32 FInternalNetRefIndex;
+
 	struct FReplicationInstanceProtocol;
 	struct FReplicationProtocol;
+
 	class FNetRefHandle;
+	class FInternalNetRefIndexManager;
 
 	namespace Private
 	{
-		typedef uint32 FInternalNetRefIndex;
 		class FNetRefHandleManager;
 	}
 }
@@ -72,6 +77,9 @@ struct FNetObjectFilteringParams
 
 	/** List of objects that have been filtered out by groups for the ConnectionId */
 	const UE::Net::FNetBitArrayView GroupFilteredOutObjects;
+
+	/** Gives access to some Iris internal bitarrays exposing status flags for each netobject */
+	const UE::Net::FInternalNetRefIndexManager& NetRefIndexManager;
 };
 
 /**
@@ -84,6 +92,9 @@ struct FNetObjectPreFilteringParams
 
 	/** FilteringInfos for all objects. Index using the set bit indices in FilteredObjects. */
 	TArrayView<const FNetObjectFilteringInfo> FilteringInfos;
+
+	/** Gives access to some Iris internal bitarrays exposing status flags for each netobject */
+	const UE::Net::FInternalNetRefIndexManager& NetRefIndexManager;
 };
 
 /**
@@ -134,7 +145,7 @@ public:
 	/** The maximum number of replicated objects in the system. */
 	uint32 AbsoluteMaxNetObjectCount = 0;
 	/** The current maximum replicated objects referenced by an index (may grow at runtime). */
-	uint32 CurrentMaxInternalIndex = 0;
+	UE::Net::FInternalNetRefIndex CurrentMaxInternalIndex = 0;
 	/** The maximum number of connections in the system. */
 	uint32 MaxConnectionCount = 0;
 };
@@ -153,6 +164,9 @@ struct FNetObjectFilterAddObjectParams
 	/** The FReplicationProtocol which describes the internal state data. */
 	const UE::Net::FReplicationProtocol* Protocol;
 
+	/** Gives access to some Iris internal bitarrays exposing status flags for each netobject */
+	const UE::Net::FInternalNetRefIndexManager& NetRefIndexManager;
+
 	/**
 	 * One can retrieve relevant information from the object state buffer using the FReplicationProtocol.
 	 * Note that this is the internal network representation of the data which is stored in quantized form.
@@ -164,8 +178,9 @@ struct FNetObjectFilterAddObjectParams
 /** Parameters passed to the filter's UpdateObjects() call. */
 struct FNetObjectFilterUpdateParams
 {
+	//TODO: This should be a TConstArrayView
 	/** Indices of the updated objects. */
-	const uint32* ObjectIndices = nullptr;
+	const UE::Net::FInternalNetRefIndex* ObjectIndices = nullptr;
 	/** The number of objects that have been updated. */
 	uint32 ObjectCount = 0;
 
@@ -182,7 +197,7 @@ public:
 	IRISCORE_API void Init(const FNetObjectFilterInitParams& Params);
 	IRISCORE_API void Deinit();
 
-	IRISCORE_API void MaxInternalNetRefIndexIncreased(UE::Net::Private::FInternalNetRefIndex MaxInternalIndex, TArrayView<FNetObjectFilteringInfo> NewFilterInfoView);
+	IRISCORE_API void MaxInternalNetRefIndexIncreased(UE::Net::FInternalNetRefIndex MaxInternalIndex, TArrayView<FNetObjectFilteringInfo> NewFilterInfoView);
 	
 	/** A new connection has been added. An opportunity for the filter to allocate per connection info. */
 	IRISCORE_API virtual void AddConnection(uint32 ConnectionId);
@@ -191,10 +206,10 @@ public:
 	IRISCORE_API virtual void RemoveConnection(uint32 ConnectionId);
 
 	/** A new object want to use this filter. Opportunity to cache some information for it. The info struct passed has been zeroed. Must be overriden. */
-	virtual bool AddObject(uint32 ObjectIndex, FNetObjectFilterAddObjectParams&) PURE_VIRTUAL(AddObject, return false;)
+	virtual bool AddObject(UE::Net::FInternalNetRefIndex ObjectIndex, FNetObjectFilterAddObjectParams&) PURE_VIRTUAL(AddObject, return false;)
 
 	/** An object no longer wants to use this filter. */
-	virtual void RemoveObject(uint32 ObjectIndex, const FNetObjectFilteringInfo&) PURE_VIRTUAL(RemoveObject,)
+	virtual void RemoveObject(UE::Net::FInternalNetRefIndex ObjectIndex, const FNetObjectFilteringInfo&) PURE_VIRTUAL(RemoveObject,)
 
 	/** A set of objects using this filter are dirty since the last update. An opportunity for the filter to update cached data. */
 	IRISCORE_API virtual void UpdateObjects(FNetObjectFilterUpdateParams&);
@@ -247,16 +262,16 @@ protected:
 	virtual void OnDeinit() PURE_VIRTUAL(OnDeinit);
 
 	/** Called when the maximum InternalNetRefIndex increased and we need to realloc our lists */
-	virtual void OnMaxInternalNetRefIndexIncreased(uint32 NewMaxInternalIndex) PURE_VIRTUAL(OnMaxInternalNetRefIndexIncreased);
+	virtual void OnMaxInternalNetRefIndexIncreased(UE::Net::FInternalNetRefIndex NewMaxInternalIndex) PURE_VIRTUAL(OnMaxInternalNetRefIndexIncreased);
 
 	/* Returns the filtering info for this object if it's handled by this filter, nullptr otherwise. */
-	IRISCORE_API FNetObjectFilteringInfo* GetFilteringInfo(uint32 ObjectIndex);
+	IRISCORE_API FNetObjectFilteringInfo* GetFilteringInfo(UE::Net::FInternalNetRefIndex ObjectIndex);
 
-	/* Returns the object index for the given NetRefHandle */
-	IRISCORE_API uint32 GetObjectIndex(UE::Net::FNetRefHandle NetRefHandle) const;
+	/* Returns the internal net ref index for the given NetRefHandle */
+	IRISCORE_API UE::Net::FInternalNetRefIndex GetObjectIndex(UE::Net::FNetRefHandle NetRefHandle) const;
 
 	/* Returns true if the object is assigned to be filtered by this filter.*/
-	inline bool IsObjectFiltered(uint32 ObjectIndex) const;
+	inline bool IsObjectFiltered(UE::Net::FInternalNetRefIndex ObjectIndex) const;
 
 	/** Adds traits. */
 	inline void AddFilterTraits(ENetFilterTraits Traits);
@@ -277,7 +292,7 @@ private:
 	TArrayView<FNetObjectFilteringInfo> FilteringInfos; 
 };
 
-inline bool UNetObjectFilter::IsObjectFiltered(uint32 ObjectIndex) const
+inline bool UNetObjectFilter::IsObjectFiltered(UE::Net::FInternalNetRefIndex ObjectIndex) const
 {
 	return ObjectIndex < FilteredObjects.GetNumBits() && FilteredObjects.IsBitSet(ObjectIndex);
 }
diff --git a/Engine/Source/Runtime/Net/Iris/Public/Iris/ReplicationSystem/Filtering/NetObjectGridFilter.h b/Engine/Source/Runtime/Net/Iris/Public/Iris/ReplicationSystem/Filtering/NetObjectGridFilter.h
--- a/Engine/Source/Runtime/Net/Iris/Public/Iris/ReplicationSystem/Filtering/NetObjectGridFilter.h
+++ b/Engine/Source/Runtime/Net/Iris/Public/Iris/ReplicationSystem/Filtering/NetObjectGridFilter.h
@@ -2,12 +2,18 @@
 
 #pragma once
 
-#include "Iris/ReplicationSystem/Filtering/NetObjectFilter.h"
+#include "Containers/ChunkedArray.h"
+
 #include "Iris/IrisConfig.h"
+#include "Iris/ReplicationSystem/InternalNetRefIndexManager.h"
+#include "Iris/ReplicationSystem/Filtering/NetObjectFilter.h"
+
 #include "Net/Core/NetBitArray.h"
-#include "Containers/ChunkedArray.h"
+
 #include "Math/Vector.h"
+
 #include "UObject/StrongObjectPtr.h"
+
 #include "NetObjectGridFilter.generated.h"
 
 namespace UE::Net
@@ -87,15 +93,15 @@ protected:
 	// UNetObjectFilter interface
 	IRISCORE_API virtual void OnInit(const FNetObjectFilterInitParams&) override;
 	IRISCORE_API virtual void OnDeinit() override;
-	IRISCORE_API virtual void OnMaxInternalNetRefIndexIncreased(uint32 NewMaxInternalIndex) override;
+	IRISCORE_API virtual void OnMaxInternalNetRefIndexIncreased(UE::Net::FInternalNetRefIndex NewMaxInternalIndex) override;
 	IRISCORE_API virtual void AddConnection(uint32 ConnectionId) override;
 	IRISCORE_API virtual void RemoveConnection(uint32 ConnectionId) override;
-	IRISCORE_API virtual bool AddObject(uint32 ObjectIndex, FNetObjectFilterAddObjectParams&) override;
-	IRISCORE_API virtual void RemoveObject(uint32 ObjectIndex, const FNetObjectFilteringInfo&) override;
+	IRISCORE_API virtual bool AddObject(UE::Net::FInternalNetRefIndex ObjectNetIndex, FNetObjectFilterAddObjectParams&) override;
+	IRISCORE_API virtual void RemoveObject(UE::Net::FInternalNetRefIndex ObjectNetIndex, const FNetObjectFilteringInfo&) override;
 	IRISCORE_API virtual void PreFilter(FNetObjectPreFilteringParams&) override;
 	IRISCORE_API virtual void Filter(FNetObjectFilteringParams&) override;
 	IRISCORE_API virtual void PostFilter(FNetObjectPostFilteringParams&) override;
-	IRISCORE_API virtual FString PrintDebugInfoForObject(const FDebugInfoParams& Params, uint32 ObjectIndex) const override;
+	IRISCORE_API virtual FString PrintDebugInfoForObject(const FDebugInfoParams& Params, UE::Net::FInternalNetRefIndex ObjectNetIndex) const override;
 
 protected:
 	struct FObjectLocationInfo : public FNetObjectFilteringInfo
@@ -138,7 +144,7 @@ protected:
 	{
 		FVector Position = FVector::ZeroVector;
 		FCellBox CellBox;
-		uint32 ObjectIndex = 0U;
+		UE::Net::FInternalNetRefIndex ObjectIndex = UE::Net::InvalidInternalNetRefIndex;
 		uint16 FrameCountBeforeCulling = 0U;
 
 		float GetCullDistance() const
@@ -172,10 +178,10 @@ protected:
 	virtual void UpdateObjectInfo(FPerObjectInfo& PerObjectInfo, const FObjectLocationInfo& ObjectLocationInfo, const UE::Net::FReplicationInstanceProtocol* InstanceProtocol) {}
 
 	/** Build the data needed to properly filter this object. */
-	virtual bool BuildObjectInfo(uint32 ObjectIndex, FNetObjectFilterAddObjectParams& Params) { return false; }
+	virtual bool BuildObjectInfo(UE::Net::FInternalNetRefIndex ObjectNetIndex, FNetObjectFilterAddObjectParams& Params) { return false; }
 
 	/** Callback when an object is removed from the grid. Useful to cleanup data tied to the object. */
-	virtual void OnObjectRemoved(uint32 ObjectIndex) {}
+	virtual void OnObjectRemoved(UE::Net::FInternalNetRefIndex ObjectNetIndex) {}
 
 private:
 
@@ -198,7 +204,7 @@ private:
 
 	struct FCellObjects
 	{
-		TSet<uint32> ObjectIndices;
+		TSet<UE::Net::FInternalNetRefIndex> ObjectIndices;
 	};
 
 	struct FCellAndTimestamp
@@ -219,7 +225,7 @@ private:
 		TArray<FCellAndTimestamp, TInlineAllocator<32>> RecentCells;
 		
 		// Objects that have been recently visible to the connection and the frame countdown.
-		TMap<uint32, uint16> RecentObjectFrameCount;
+		TMap<UE::Net::FInternalNetRefIndex, uint16> RecentObjectFrameCount;
 	};
 
 	/** Aggregator for stats */
diff --git a/Engine/Source/Runtime/Net/Iris/Public/Iris/ReplicationSystem/InternalNetRefIndexManager.h b/Engine/Source/Runtime/Net/Iris/Public/Iris/ReplicationSystem/InternalNetRefIndexManager.h
new file mode 100644
--- /dev/null
+++ b/Engine/Source/Runtime/Net/Iris/Public/Iris/ReplicationSystem/InternalNetRefIndexManager.h
@@ -0,0 +1,45 @@
+// Copyright Epic Games, Inc. All Rights Reserved.
+
+#pragma once
+
+#include "HAL/Platform.h"
+
+namespace UE::Net::Private
+{
+	class FNetRefHandleManager;
+}
+
+namespace UE::Net
+{
+	/** Internal type to visualize indexes mapping a NetRefHandle to their position in internal NetBitArrays. */
+	typedef uint32 FInternalNetRefIndex;
+
+	/** The invalid index is always the first index since it will never be assigned. */
+	constexpr uint32 InvalidInternalNetRefIndex = 0U;
+}
+
+namespace UE::Net
+{
+
+/** Returns true if the NetIndex is not the default invalid value */
+static bool IsValidInternalNetRefIndex(const FInternalNetRefIndex NetIndex)
+{
+	return NetIndex != InvalidInternalNetRefIndex;
+}
+
+/**
+ * This class is a facade over the NetRefHandleManager to allow public access to some internal Iris netbitarrays.
+ */
+class FInternalNetRefIndexManager
+{
+public:
+
+	FInternalNetRefIndexManager() = delete;
+	FInternalNetRefIndexManager(const UE::Net::Private::FNetRefHandleManager& InNetRefHandleManager);
+	
+private:
+
+	const UE::Net::Private::FNetRefHandleManager& NetRefHandleManager;
+};
+
+} // end namespace UE::Net
\ No newline at end of file
diff --git a/Engine/Source/Runtime/Net/Iris/Public/Iris/ReplicationSystem/ObjectReplicationBridge.h b/Engine/Source/Runtime/Net/Iris/Public/Iris/ReplicationSystem/ObjectReplicationBridge.h
--- a/Engine/Source/Runtime/Net/Iris/Public/Iris/ReplicationSystem/ObjectReplicationBridge.h
+++ b/Engine/Source/Runtime/Net/Iris/Public/Iris/ReplicationSystem/ObjectReplicationBridge.h
@@ -41,8 +41,6 @@ namespace UE::Net
 
 	namespace Private
 	{
-		typedef uint32 FInternalNetRefIndex;
-
 		enum class EInternalDetachReason : uint8;
 
 		class FObjectPollFrequencyLimiter;
@@ -75,7 +73,7 @@ struct FStaticDestructionInfo
 	UE::Net::FNetObjectFactoryId NetFactoryId = UE::Net::InvalidNetObjectFactoryId;
 
 	/** The repindex assigned to this destruction info */
-	UE::Net::Private::FInternalNetRefIndex InternalReplicationIndex;
+	UE::Net::FInternalNetRefIndex InternalReplicationIndex;
 };
 
 /** 
@@ -402,7 +400,7 @@ protected:
 	IRISCORE_API virtual void OnProtocolMismatchReported(FNetRefHandle RefHandle, uint32 ConnectionId) override;
 	IRISCORE_API virtual void OnErrorWithNetRefHandleReported(UE::Net::ENetRefHandleError ErrorType, FNetRefHandle RefHandle, uint32 ConnectionId, const TArray<FNetRefHandle>& ExtraHandles, const FString& DiagnosticMessage) override;
 	IRISCORE_API virtual void RemoveDestructionInfosForGroup(UE::Net::FNetObjectGroupHandle GroupHandle) override;
-	IRISCORE_API virtual void RemoveDestructionInfo(UE::Net::Private::FInternalNetRefIndex InternalObjectIndex) override;
+	IRISCORE_API virtual void RemoveDestructionInfo(UE::Net::FInternalNetRefIndex InternalObjectIndex) override;
 
 private:
 
@@ -443,7 +441,7 @@ protected:
 	IRISCORE_API void LoadConfig();
 
 	/** Force the specified object to update it's cached WorldLocation data. */
-	IRISCORE_API void ForceUpdateWorldLocation(FNetRefHandle NetRefHandle, UE::Net::Private::FInternalNetRefIndex InternalObjectIndex);
+	IRISCORE_API void ForceUpdateWorldLocation(FNetRefHandle NetRefHandle, UE::Net::FInternalNetRefIndex InternalObjectIndex);
 
 	/**
 	 * Set the function used to determine whether an object should use the default spatial filter
@@ -545,7 +543,7 @@ private:
 
 	/** Detach a subobject on the client due to a request from the authority */
 	//DEPRECATED: Will be replaced with DetachRootObjectFromRemote soon
-	void DetachSubObjectInstancesFromRemote(UE::Net::Private::FInternalNetRefIndex RootObjectIndex, EReplicationBridgeDestroyInstanceReason DestroyReason, EReplicationBridgeDestroyInstanceFlags DestroyFlags);
+	void DetachSubObjectInstancesFromRemote(UE::Net::FInternalNetRefIndex RootObjectIndex, EReplicationBridgeDestroyInstanceReason DestroyReason, EReplicationBridgeDestroyInstanceFlags DestroyFlags);
 
 	void SetNetPushIdOnInstance(UE::Net::FReplicationInstanceProtocol* InstanceProtocol, FNetHandle NetHandle);
 
@@ -553,7 +551,7 @@ private:
 	void FindClassesInPollPeriodOverrides();
 
 	/** Returns the poll frequency of a specific root object. */
-	float GetPollFrequencyOfRootObject(UE::Net::Private::FInternalNetRefIndex RootObjectIndex) const;
+	float GetPollFrequencyOfRootObject(UE::Net::FInternalNetRefIndex RootObjectIndex) const;
 
 	/** Does all that is necessary to replicate any given UObject. Returns the object's handle if successful  */
 	UE::Net::FNetRefHandle StartReplicatingNetObject(UObject* Instance, UE::Net::EReplicationFragmentTraits Traits, UE::Net::FNetObjectFactoryId NetFactoryId, bool bForceNewHandle);
@@ -576,7 +574,7 @@ private:
 	struct FUpdateWorldInfoContext
 	{
 		UE::Net::FWorldLocations& WorldLocations;
-		UE::Net::Private::FInternalNetRefIndex ObjectIndex;
+		UE::Net::FInternalNetRefIndex ObjectIndex;
 		FNetRefHandle NetRefHandle;
 	};
 	/** Update the WorldLocationInfo of a specific root object */
@@ -594,7 +592,7 @@ private:
 	void InitNetObjectFactories();
 	void DeinitNetObjectFactories();
 
-	void OnMaxInternalNetRefIndexIncreased(UE::Net::Private::FInternalNetRefIndex NewMaxInternalIndex);
+	void OnMaxInternalNetRefIndexIncreased(UE::Net::FInternalNetRefIndex NewMaxInternalIndex);
 
 private:
 
@@ -616,14 +614,14 @@ private:
 	 * @param Handle The handle of the object to write creation data for.
 	 * @param InternalObjectIndex the internalindex we are currently writing creation info for.
 	 */
-	bool WriteNetRefHandleCreationInfo(FReplicationBridgeSerializationContext& Context, FNetRefHandle Handle, UE::Net::Private::FInternalNetRefIndex InternalObjectIndex);
+	bool WriteNetRefHandleCreationInfo(FReplicationBridgeSerializationContext& Context, FNetRefHandle Handle, UE::Net::FInternalNetRefIndex InternalObjectIndex);
 
 	/**
 	 * Client only function that executes the EndReplication request by the authority
 	 * @param InternalObjectIndex The index of the object to destroy
 	 * @DestroyReason How the object should be destroyed
 	 */
-	void DestroyNetObjectFromRemote(UE::Net::Private::FInternalNetRefIndex InternalObjectIndex, EReplicationBridgeDestroyInstanceReason DestroyReason, UE::Net::Private::EInternalDetachReason InternalDetachReason);
+	void DestroyNetObjectFromRemote(UE::Net::FInternalNetRefIndex InternalObjectIndex, EReplicationBridgeDestroyInstanceReason DestroyReason, UE::Net::Private::EInternalDetachReason InternalDetachReason);
 
 	void ReadAndExecuteDestructionInfoFromRemote(FReplicationBridgeSerializationContext& Context);
 
@@ -631,10 +629,10 @@ private:
 	FReplicationBridgeCreateNetRefHandleResult CreateNetRefHandleFromRemote(FNetRefHandle RootObjectNetHandle, FNetRefHandle WantedNetHandle, FReplicationBridgeSerializationContext& Context);
 
 	/** Invoked after we have applied the initial state for an object.*/
-	void PostApplyInitialState(UE::Net::Private::FInternalNetRefIndex InternalObjectIndex);
+	void PostApplyInitialState(UE::Net::FInternalNetRefIndex InternalObjectIndex);
 
 	/** Invoked right before we apply the state for a new received subobject but after we have applied state on the root object. */
-	void SubObjectCreatedFromReplication(UE::Net::Private::FInternalNetRefIndex RootObjectIndex, FNetRefHandle SubObjectCreated);
+	void SubObjectCreatedFromReplication(UE::Net::FInternalNetRefIndex RootObjectIndex, FNetRefHandle SubObjectCreated);
 
 	/** Gets the bReplicateTransactionally flag from the ReplicationSystem */
 	bool IsReplicatingTransactionally() const;
@@ -858,7 +856,7 @@ private:
 		return Bridge->PostApplyInitialState(InternalObjectIndex);
 	}
 
-	static void SubObjectCreatedFromReplication(UObjectReplicationBridge* Bridge, UE::Net::Private::FInternalNetRefIndex RootObjectIndex, FNetRefHandle SubObjectCreated)
+	static void SubObjectCreatedFromReplication(UObjectReplicationBridge* Bridge, UE::Net::FInternalNetRefIndex RootObjectIndex, FNetRefHandle SubObjectCreated)
 	{
 		return Bridge->SubObjectCreatedFromReplication(RootObjectIndex, SubObjectCreated);
 	}
diff --git a/Engine/Source/Runtime/Net/Iris/Public/Iris/ReplicationSystem/Prioritization/NetObjectPrioritizer.h b/Engine/Source/Runtime/Net/Iris/Public/Iris/ReplicationSystem/Prioritization/NetObjectPrioritizer.h
--- a/Engine/Source/Runtime/Net/Iris/Public/Iris/ReplicationSystem/Prioritization/NetObjectPrioritizer.h
+++ b/Engine/Source/Runtime/Net/Iris/Public/Iris/ReplicationSystem/Prioritization/NetObjectPrioritizer.h
@@ -9,9 +9,12 @@
 
 struct FNetObjectPrioritizationInfo;
 class UReplicationSystem;
+
 namespace UE::Net
 {
 	typedef uint32 FNetObjectPrioritizerHandle;
+	typedef uint32 FInternalNetRefIndex;
+
 	struct FReplicationInstanceProtocol;
 	struct FReplicationProtocol;
 	struct FReplicationView;
@@ -42,15 +45,17 @@ constexpr FNetObjectPrioritizerHandle DefaultSpatialNetObjectPrioritizerHandle =
   */
 struct FNetObjectPrioritizationParams
 {
+	//TODO: this should be a TConstArrayView
 	/** 
 	 * The indices for the objects that are being prioritized. Will only contain objects
 	 * which have been added to the prioritizer and have dirty properties or are in need of resending.
 	 */
-	const uint32* ObjectIndices;
+	const UE::Net::FInternalNetRefIndex* ObjectIndices;
 
 	/** The number of objects to prioritize. */
 	uint32 ObjectCount;
 
+	//TODO: This should be a TArrayView
 	/** Priorities for all objects. Index using ObjectIndices[0..ObjectCount-1]. */
 	float* Priorities;
 
@@ -131,8 +136,9 @@ struct FNetObjectPrioritizerAddObjectParams
 /** Parameters passed to the prioritizer's UpdateObjects() call. */
 struct FNetObjectPrioritizerUpdateParams
 {
+	//TODO: this should be a TConstArrayView
 	/** Indices of the updated objects. */
-	const uint32* ObjectIndices;
+	const UE::Net::FInternalNetRefIndex* ObjectIndices;
 	/** The number of objects that have been updated. */
 	uint32 ObjectCount;
 
@@ -166,7 +172,7 @@ public:
 	virtual void Deinit() PURE_VIRTUAL(Deinit)
 
 	/** Called when the maximum InternalNetRefIndex increased and we need to realloc our lists */
-	virtual void OnMaxInternalNetRefIndexIncreased(uint32 NewMaxInternalIndex) PURE_VIRTUAL(OnMaxInternalNetRefIndexIncreased);
+	virtual void OnMaxInternalNetRefIndexIncreased(UE::Net::FInternalNetRefIndex NewMaxInternalIndex) PURE_VIRTUAL(OnMaxInternalNetRefIndexIncreased);
 
 	/** A new connection has been added. An opportunity for the prioritizer to allocate per connection info. */
 	IRISCORE_API virtual void AddConnection(uint32 ConnectionId);
@@ -175,10 +181,10 @@ public:
 	IRISCORE_API virtual void RemoveConnection(uint32 ConnectionId);
 
 	/** A new object want to use this prioritizer. Opportunity to cache some information for it. The info struct passed has been zeroed. */
-	virtual bool AddObject(uint32 ObjectIndex, FNetObjectPrioritizerAddObjectParams& Params) PURE_VIRTUAL(AddObject, return false;)
+	virtual bool AddObject(UE::Net::FInternalNetRefIndex ObjectIndex, FNetObjectPrioritizerAddObjectParams& Params) PURE_VIRTUAL(AddObject, return false;)
 
 	/** An object do no longer want to use this prioritizer. */
-	virtual void RemoveObject(uint32 ObjectIndex, const FNetObjectPrioritizationInfo& Info) PURE_VIRTUAL(RemoveObject,)
+	virtual void RemoveObject(UE::Net::FInternalNetRefIndex ObjectIndex, const FNetObjectPrioritizationInfo& Info) PURE_VIRTUAL(RemoveObject,)
 
 	/** A set of objects used by this prioritizer have been updated. An opportunity for the prioritizer to update cached data. */
 	virtual void UpdateObjects(FNetObjectPrioritizerUpdateParams&) PURE_VIRTUAL(UpdateObjects,)
diff --git a/Engine/Source/Runtime/Net/Iris/Public/Iris/ReplicationSystem/ReplicationBridge.h b/Engine/Source/Runtime/Net/Iris/Public/Iris/ReplicationSystem/ReplicationBridge.h
--- a/Engine/Source/Runtime/Net/Iris/Public/Iris/ReplicationSystem/ReplicationBridge.h
+++ b/Engine/Source/Runtime/Net/Iris/Public/Iris/ReplicationSystem/ReplicationBridge.h
@@ -35,6 +35,8 @@ namespace UE::Net
 {
 	enum class ENetRefHandleError : uint32;
 
+	typedef uint32 FInternalNetRefIndex;
+
 	struct FNetDependencyInfo;
 	struct FReplicationInstanceProtocol;
 	struct FReplicationProtocol;
@@ -47,7 +49,6 @@ namespace UE::Net
 
 	namespace Private
 	{
-		typedef uint32 FInternalNetRefIndex;
 		class FNetRefHandleManager;
 		class FNetObjectGroups;
 		class FNetPushObjectHandle;
@@ -229,10 +230,10 @@ protected:
 	IRISCORE_API void InternalAttachInstanceToNetRefHandle(FNetRefHandle RefHandle, bool bBindInstanceProtocol, UE::Net::FReplicationInstanceProtocol* InstanceProtocol, UObject* Instance, FNetHandle NetHandle);
 
 	/** Detach instance from NetRefHandle and destroy the instance protocol. */
-	IRISCORE_API void InternalDetachInstanceFromNetRefHandle(UE::Net::Private::FInternalNetRefIndex InternalObjectIndex, bool bUnBindInstanceProtocol = true);
+	IRISCORE_API void InternalDetachInstanceFromNetRefHandle(UE::Net::FInternalNetRefIndex InternalObjectIndex, bool bUnBindInstanceProtocol = true);
 
 	/** Destroy the handle and all internal book keeping associated with it. */
-	IRISCORE_API void InternalDestroyNetObject(UE::Net::Private::FInternalNetRefIndex InternalObjectIndex);
+	IRISCORE_API void InternalDestroyNetObject(UE::Net::FInternalNetRefIndex InternalObjectIndex);
 	
 	/**
 	 * Add SubObjectHandle as SubObject to ParentHandle. If the ParentHandle is a SubObject itself the SubObject will be share the same RootObject but will be subjects to the Parents conditionals.
@@ -281,8 +282,8 @@ protected:
 	 */
 	virtual void RemoveDestructionInfosForGroup(UE::Net::FNetObjectGroupHandle GroupHandle) {}
 
-	/** Remove destruction infos associated with a specific NetRefHandle */
-	virtual void RemoveDestructionInfo(UE::Net::Private::FInternalNetRefIndex InternalObjectIndex) {};
+	/** Remove destruction infos associated to a net object*/
+	virtual void RemoveDestructionInfo(UE::Net::FInternalNetRefIndex InternalObjectIndex) {};
 
 private:
 
@@ -296,7 +297,7 @@ private:
 
 	// Adds the Handle to the list of handles pending deferred EndReplication, if bIsImmediate is true the object will be destroyed after the next update, otherwise
 	// it will be kept around until the handle is no longer ref-counted by any connection. It will however be removed from the set of scopeable objects after the first update so new connections will not add it to their scope.
-	void AddPendingEndReplication(FNetRefHandle Handle, UE::Net::Private::FInternalNetRefIndex InternalObjectIndex, EEndReplicationFlags DestroyFlags, EPendingEndReplicationImmediate Immediate = EPendingEndReplicationImmediate::No);
+	void AddPendingEndReplication(FNetRefHandle Handle, UE::Net::FInternalNetRefIndex InternalObjectIndex, EEndReplicationFlags DestroyFlags, EPendingEndReplicationImmediate Immediate = EPendingEndReplicationImmediate::No);
 
 	void CallPreSendUpdate(float DeltaSeconds);
 	void CallPreSendUpdateSingleHandle(FNetRefHandle Handle);
@@ -310,15 +311,15 @@ private:
 
 private:
 
-	void InternalFlushStateData(UE::Net::FNetSerializationContext& SerializationContext, UE::Net::Private::FChangeMaskCache& ChangeMaskCache, UE::Net::FNetBitStreamWriter& ChangeMaskWriter, UE::Net::Private::FInternalNetRefIndex InternalObjectIndex);
+	void InternalFlushStateData(UE::Net::FNetSerializationContext& SerializationContext, UE::Net::Private::FChangeMaskCache& ChangeMaskCache, UE::Net::FNetBitStreamWriter& ChangeMaskWriter, UE::Net::FInternalNetRefIndex InternalObjectIndex);
 	// Internal method to copy state data for Handle
-	void InternalFlushStateData(UE::Net::Private::FInternalNetRefIndex InternalObjectIndex);
+	void InternalFlushStateData(UE::Net::FInternalNetRefIndex InternalObjectIndex);
 
 	// Internal method to copy state data for Handle and any SubObjects and mark them as being torn-off
 	void InternalTearOff(FNetRefHandle OwnerHandle);
 
 	// Internal method to stop replication by filtering out the replicated object. Once it is no longer replicating to any client it will be destroyed.
-	void InternalRequestAsyncStopReplication(UE::Net::Private::FInternalNetRefIndex InternalObjectIndex, EEndReplicationFlags Flags);
+	void InternalRequestAsyncStopReplication(UE::Net::FInternalNetRefIndex InternalObjectIndex, EEndReplicationFlags Flags);
 
 	/**
 	 * Called from ReplicationSystem when a streaming container is about to unload.
@@ -326,13 +327,13 @@ private:
 	 */
 	void NotifyStreamingContainerUnload(const UObject* Container);
 
-	void RemoveFromAllGroups(UE::Net::Private::FInternalNetRefIndex InternalObjectIndex);
+	void RemoveFromAllGroups(UE::Net::FInternalNetRefIndex InternalObjectIndex);
 
 	// Destroy is split into two phases in order to support replicating data for objects where the external representation might have been destroyed.
-	void InternalStartDestroyLocalSubObjects(UE::Net::Private::FInternalNetRefIndex InternalIndex, EEndReplicationFlags Flags);
-	void InternalFinishDestroyLocalSubObjects(UE::Net::Private::FInternalNetRefIndex InternalIndex, EEndReplicationFlags Flags);
-	void InternalStartDestroyLocalObject(UE::Net::Private::FInternalNetRefIndex InternalIndex, EEndReplicationFlags Flags);
-	void InternalFinishDestroyLocalObject(UE::Net::Private::FInternalNetRefIndex InternalIndex, EEndReplicationFlags Flags);
+	void InternalStartDestroyLocalSubObjects(UE::Net::FInternalNetRefIndex InternalIndex, EEndReplicationFlags Flags);
+	void InternalFinishDestroyLocalSubObjects(UE::Net::FInternalNetRefIndex InternalIndex, EEndReplicationFlags Flags);
+	void InternalStartDestroyLocalObject(UE::Net::FInternalNetRefIndex InternalIndex, EEndReplicationFlags Flags);
+	void InternalFinishDestroyLocalObject(UE::Net::FInternalNetRefIndex InternalIndex, EEndReplicationFlags Flags);
 
 	// Destroy local NetObject associated with handle
 	void DestroyLocalNetHandle(FNetRefHandle Handle, EEndReplicationFlags Flags);
@@ -347,8 +348,8 @@ private:
 	void SetNetPushIdOnFragments(const TArrayView<const UE::Net::FReplicationFragment*const>& Fragments, const UE::Net::Private::FNetPushObjectHandle& PushHandle);
 	void ClearNetPushIdOnFragments(const TArrayView<const UE::Net::FReplicationFragment*const>& Fragments);
 
-	void DestroyGlobalNetHandle(UE::Net::Private::FInternalNetRefIndex InternalReplicationIndex);
-	void ClearNetPushIds(UE::Net::Private::FInternalNetRefIndex InternalReplicationIndex);
+	void DestroyGlobalNetHandle(UE::Net::FInternalNetRefIndex InternalReplicationIndex);
+	void ClearNetPushIds(UE::Net::FInternalNetRefIndex InternalReplicationIndex);
 
 	friend UE::Net::Private::FReplicationSystemImpl;
 	friend UE::Net::Private::FReplicationSystemInternal;
@@ -375,7 +376,7 @@ private:
 
 	struct FPendingEndReplicationInfo
 	{
-		FPendingEndReplicationInfo(FNetRefHandle InHandle, UE::Net::Private::FInternalNetRefIndex InObjectIndex, EEndReplicationFlags InEndReplicationFlags, EPendingEndReplicationImmediate InImmediate)
+		FPendingEndReplicationInfo(FNetRefHandle InHandle, UE::Net::FInternalNetRefIndex InObjectIndex, EEndReplicationFlags InEndReplicationFlags, EPendingEndReplicationImmediate InImmediate)
 		: Handle(InHandle)
 		, ObjectIndex(InObjectIndex)
 		, EndReplicationFlags(InEndReplicationFlags)
@@ -384,7 +385,7 @@ private:
 		}
 
 		FNetRefHandle Handle;
-		UE::Net::Private::FInternalNetRefIndex ObjectIndex;
+		UE::Net::FInternalNetRefIndex ObjectIndex;
 		EEndReplicationFlags EndReplicationFlags;
 		EPendingEndReplicationImmediate Immediate;
 	};
diff --git a/Engine/Source/Runtime/Net/Iris/Public/Iris/ReplicationSystem/WorldLocations.h b/Engine/Source/Runtime/Net/Iris/Public/Iris/ReplicationSystem/WorldLocations.h
--- a/Engine/Source/Runtime/Net/Iris/Public/Iris/ReplicationSystem/WorldLocations.h
+++ b/Engine/Source/Runtime/Net/Iris/Public/Iris/ReplicationSystem/WorldLocations.h
@@ -12,10 +12,14 @@
 
 class UReplicationSystem;
 
-namespace UE::Net::Private
+namespace UE::Net
 {
 	typedef uint32 FInternalNetRefIndex;
-	class FNetRefHandleManager;
+
+	namespace Private
+	{
+		class FNetRefHandleManager;
+	}
 }
 
 /**
@@ -47,7 +51,7 @@ struct FWorldLocationsInitParams
 {
 	TObjectPtr<UReplicationSystem> ReplicationSystem = nullptr;
 
-	UE::Net::Private::FInternalNetRefIndex MaxInternalNetRefIndex = 0;
+	FInternalNetRefIndex MaxInternalNetRefIndex = 0;
 
 	/** How many world info storage slots to preallocate. */
 	uint32 PreallocatedStorageCount = 0;
@@ -73,41 +77,41 @@ public:
 	void PostSendUpdate();
 
 	/** Returns whether the object has a valid cached data or not. */
-	bool HasInfoForObject(UE::Net::Private::FInternalNetRefIndex ObjectIndex) const
+	bool HasInfoForObject(FInternalNetRefIndex ObjectIndex) const
 	{
 		return ValidInfoIndexes.IsBitSet(ObjectIndex);
 	}
 
 	/** Returns the object's world location if it's valid or a zero vector if it's not. */
-	inline FVector GetWorldLocation(UE::Net::Private::FInternalNetRefIndex ObjectIndex) const;
+	inline FVector GetWorldLocation(FInternalNetRefIndex ObjectIndex) const;
 
 	/** Get the object's current cull distance. */
-	inline float GetCullDistance(UE::Net::Private::FInternalNetRefIndex ObjectIndex) const;
+	inline float GetCullDistance(FInternalNetRefIndex ObjectIndex) const;
 
 	/** Return the current stored world information of the given object, or NullOpt if the object did not register in the world location cache. */
-	inline TOptional<FWorldLocations::FWorldInfo> GetWorldInfo(UE::Net::Private::FInternalNetRefIndex ObjectIndex) const;
+	inline TOptional<FWorldLocations::FWorldInfo> GetWorldInfo(FInternalNetRefIndex ObjectIndex) const;
 
 	/** Set the mandatory info of a replicated root object*/
-	void SetObjectInfo(UE::Net::Private::FInternalNetRefIndex ObjectIndex, const FVector& Location, float NetCullDistance);
+	void SetObjectInfo(FInternalNetRefIndex ObjectIndex, const FVector& Location, float NetCullDistance);
 	
 	/** Assign a world information cache to the replicated object */
-	void InitObjectInfoCache(UE::Net::Private::FInternalNetRefIndex ObjectIndex);
+	void InitObjectInfoCache(FInternalNetRefIndex ObjectIndex);
 
 	/** Remove the world information cache of the replicated object */
-	void RemoveObjectInfoCache(UE::Net::Private::FInternalNetRefIndex ObjectIndex);
+	void RemoveObjectInfoCache(FInternalNetRefIndex ObjectIndex);
 
 	/**
 	 * Objects are not necessarily marked as dirty just because they're moving, such as objects attached to other objects. 
 	 * If such objects are spatially filtered they need to update their world locations in order for replication to work as expected.
 	 * Use SetObjectRequiresFrequentWorldLocationUpdate to force frequent world location update on an object.
 	 */
-	void SetObjectRequiresFrequentWorldLocationUpdate(UE::Net::Private::FInternalNetRefIndex ObjectIndex, bool bRequiresFrequentUpdate)
+	void SetObjectRequiresFrequentWorldLocationUpdate(FInternalNetRefIndex ObjectIndex, bool bRequiresFrequentUpdate)
 	{
 		ObjectsRequiringFrequentWorldLocationUpdate.SetBitValue(ObjectIndex, ValidInfoIndexes.GetBit(ObjectIndex) && bRequiresFrequentUpdate);
 	}
 
 	/** Returns whether an object requires frequent world location updates. */
-	bool GetObjectRequiresFrequentWorldLocationUpdate(UE::Net::Private::FInternalNetRefIndex ObjectIndex) const
+	bool GetObjectRequiresFrequentWorldLocationUpdate(FInternalNetRefIndex ObjectIndex) const
 	{
 		return ObjectsRequiringFrequentWorldLocationUpdate.GetBit(ObjectIndex);
 	}
@@ -116,16 +120,16 @@ public:
 	 * Add a temporary net cull distance that will have priority over the regular net cull distance.
 	 * Returns true if the object had registered to use the world location cache and can store the override.
 	 */
-	bool SetCullDistanceOverride(UE::Net::Private::FInternalNetRefIndex ObjectIndex, float CullDistance);
+	bool SetCullDistanceOverride(FInternalNetRefIndex ObjectIndex, float CullDistance);
 
 	/** 
 	 * Remove the temporary net cull distance override and instead use the regular net cull distance 
 	 * Returns true if the object had registered to use the world location cache and had an override value previously set
 	 */
-	bool ClearCullDistanceOverride(UE::Net::Private::FInternalNetRefIndex ObjectIndex);
+	bool ClearCullDistanceOverride(FInternalNetRefIndex ObjectIndex);
 
 	/** Returns true if the object was set a cull distance override and is using it instead of his default cull distance value */
-	bool HasCullDistanceOverride(UE::Net::Private::FInternalNetRefIndex ObjectIndex) const
+	bool HasCullDistanceOverride(FInternalNetRefIndex ObjectIndex) const
 	{
 		return ValidInfoIndexes.IsBitSet(ObjectIndex) ? GetObjectInfo(ObjectIndex).CullDistanceOverride != FLT_MAX : false;
 	}
@@ -162,7 +166,7 @@ public:
 				Location.X <= MaxWorldPos.X && Location.Y <= MaxWorldPos.Y && Location.Z <= MaxWorldPos.Z);
 	}
 
-	void OnMaxInternalNetRefIndexIncreased(UE::Net::Private::FInternalNetRefIndex NewMaxInternalIndex);
+	void OnMaxInternalNetRefIndexIncreased(FInternalNetRefIndex NewMaxInternalIndex);
 
 private:
 
@@ -186,17 +190,17 @@ private:
 
 private:
 
-	const int32 GetStorageIndex(UE::Net::Private::FInternalNetRefIndex ObjectIndex) const
+	const int32 GetStorageIndex(FInternalNetRefIndex ObjectIndex) const
 	{
 		return StorageIndexes[ObjectIndex];
 	}
 
-	const FObjectInfo& GetObjectInfo(UE::Net::Private::FInternalNetRefIndex ObjectIndex) const
+	const FObjectInfo& GetObjectInfo(FInternalNetRefIndex ObjectIndex) const
 	{
 		return StoredObjectInfo[GetStorageIndex(ObjectIndex)];
 	}
 
-	FObjectInfo& GetObjectInfo(UE::Net::Private::FInternalNetRefIndex ObjectIndex)
+	FObjectInfo& GetObjectInfo(FInternalNetRefIndex ObjectIndex)
 	{
 		return StoredObjectInfo[GetStorageIndex(ObjectIndex)];
 	}
@@ -231,12 +235,12 @@ private:
 	bool bLockdownDirtyList = false;
 };
 
-inline FVector FWorldLocations::GetWorldLocation(UE::Net::Private::FInternalNetRefIndex ObjectIndex) const
+inline FVector FWorldLocations::GetWorldLocation(FInternalNetRefIndex ObjectIndex) const
 {
 	return ValidInfoIndexes.IsBitSet(ObjectIndex) ? GetObjectInfo(ObjectIndex).WorldLocation : FVector::Zero();
 }
 
-inline float FWorldLocations::GetCullDistance(UE::Net::Private::FInternalNetRefIndex ObjectIndex) const
+inline float FWorldLocations::GetCullDistance(FInternalNetRefIndex ObjectIndex) const
 {
 	if (ValidInfoIndexes.IsBitSet(ObjectIndex))
 	{
@@ -247,7 +251,7 @@ inline float FWorldLocations::GetCullDistance(UE::Net::Private::FInternalNetRefI
 	return 0.0f;
 }
 
-inline TOptional<FWorldLocations::FWorldInfo> FWorldLocations::GetWorldInfo(UE::Net::Private::FInternalNetRefIndex ObjectIndex) const
+inline TOptional<FWorldLocations::FWorldInfo> FWorldLocations::GetWorldInfo(FInternalNetRefIndex ObjectIndex) const
 {
 	if (ValidInfoIndexes.IsBitSet(ObjectIndex))
 	{
