Iris - Suppress Clang static analyzer false positive in DataStreamManager

#Synced-CL 51584953

diff --git a/Engine/Source/Runtime/Net/Iris/Private/Iris/DataStream/DataStreamManager.cpp b/Engine/Source/Runtime/Net/Iris/Private/Iris/DataStream/DataStreamManager.cpp
--- a/Engine/Source/Runtime/Net/Iris/Private/Iris/DataStream/DataStreamManager.cpp
+++ b/Engine/Source/Runtime/Net/Iris/Private/Iris/DataStream/DataStreamManager.cpp
@@ -529,6 +529,9 @@ UDataStreamManager::EWriteResult UDataStreamManager::FImpl::WriteData(UE::Net::F
 		}
 	}
 
+	// Suppress clang analyzer false positive about Context retaining reference to ExportScope's member.
+	// The FExportScope destructor restores the original ExportContext before this return.
+	UE_CLANG_SUPPRESS
 	return CombinedWriteResult;
 }
 
