Iris - Properly free memory when quantizing

#Synced-CL 51619035

diff --git a/Engine/Source/Runtime/Net/Iris/Public/Iris/Serialization/PolymorphicNetSerializerImpl.h b/Engine/Source/Runtime/Net/Iris/Public/Iris/Serialization/PolymorphicNetSerializerImpl.h
--- a/Engine/Source/Runtime/Net/Iris/Public/Iris/Serialization/PolymorphicNetSerializerImpl.h
+++ b/Engine/Source/Runtime/Net/Iris/Public/Iris/Serialization/PolymorphicNetSerializerImpl.h
@@ -353,7 +353,7 @@ void TPolymorphicStructNetSerializerImpl<ExternalSourceType, ExternalSourceItemT
 	QuantizedType& TargetValue = *reinterpret_cast<QuantizedType*>(Args.Target);
 	const ConfigType& Config = *static_cast<const ConfigType*>(Args.NetSerializerConfig);
 
-	FPolymorphicStructNetSerializerInternal::Free(Context, TargetValue.StructData);
+	InternalFreeItem(Context, Config, TargetValue);
 
 	const TSharedPtr<SourceItemType>& Item = GetItem(SourceValue);
 	const UScriptStruct* ScriptStruct = Item.IsValid() ? Item->GetScriptStruct() : nullptr;
