Iris - Add more logging when encountering a missing destruction info.

#Synced-CL 51619035

diff --git a/Engine/Source/Runtime/Net/Iris/Private/Iris/ReplicationSystem/ReplicationWriter.cpp b/Engine/Source/Runtime/Net/Iris/Private/Iris/ReplicationSystem/ReplicationWriter.cpp
--- a/Engine/Source/Runtime/Net/Iris/Private/Iris/ReplicationSystem/ReplicationWriter.cpp
+++ b/Engine/Source/Runtime/Net/Iris/Private/Iris/ReplicationSystem/ReplicationWriter.cpp
@@ -3672,7 +3672,7 @@ bool FReplicationWriter::WriteNetRefHandleDestructionInfo(FNetSerializationConte
 	UE_NET_TRACE_SCOPE(DestructionInfo, *Context.GetBitStreamWriter(), Context.GetTraceCollector(), ENetTraceVerbosity::Trace);
 
 	const FStaticDestructionInfo* Info = FBridgeSerialization::GetStaticDestructionInfo(ReplicationBridge, Handle);
-	if (ensure(Info))
+	if (Info)
 	{
 		UE_LOG(LogIris, VeryVerbose, TEXT("WriteNetRefHandleDestructionInfo on %s | %s | NetFactoryId: %u"), *Handle.ToString(), *Info->StaticRef.ToString(), Info->NetFactoryId);
 		// Write destruction info
@@ -3681,8 +3681,14 @@ bool FReplicationWriter::WriteNetRefHandleDestructionInfo(FNetSerializationConte
 	}
 	else
 	{
+		UE_CALL_ONCE([&]()
+		{
+			UE_LOG(LogIris, Error, TEXT("Unable to write destruction info object %s in state %s representing destroyed object %s"), *NetRefHandleManager->PrintObjectFromNetRefHandle(Handle), *PrintObjectInfo(NetRefHandleManager->GetInternalIndex(Handle)), *NetRefHandleManager->PrintObjectFromIndex(NetRefHandleManager->GetOriginalDestroyedStartupObjectIndex(NetRefHandleManager->GetInternalIndex(Handle))));
+		}
+		);
+		ensureMsgf(false, TEXT("Unable to write destruction info object %s in state %s representing destroyed object %s"), *NetRefHandleManager->PrintObjectFromNetRefHandle(Handle), *PrintObjectInfo(NetRefHandleManager->GetInternalIndex(Handle)), *NetRefHandleManager->PrintObjectFromIndex(NetRefHandleManager->GetOriginalDestroyedStartupObjectIndex(NetRefHandleManager->GetInternalIndex(Handle))));
 		Context.SetError(TEXT("DestructionInfoNotFound"), false);
-		UE_LOG(LogIris, Error, TEXT("Failed to write destructionInfo for %s"), *Handle.ToString());
+		return false;
 	}
 
 	return !Context.HasErrorOrOverflow();
@@ -4743,9 +4749,10 @@ FString FReplicationWriter::PrintObjectInfo(FInternalNetRefIndex ObjectIndex) co
 
 	TStringBuilder<512> InfoBuilder;
 
-	InfoBuilder.Appendf(TEXT("Status info: 0x%llx (%s) | %s | SchedulingPriority: %f | LastAckedBaselineIndex: %u | Flags"), 
+	InfoBuilder.Appendf(TEXT("Status info (InternalIndex: %u) : 0x%llx (%s) | %s | SchedulingPriority: %f | LastAckedBaselineIndex: %u | Flags"), 
+		ObjectIndex,
 		ObjectInfo.Value, 
-		ObjectInfo.IsSubObject?TEXT("SubObject"):TEXT("RootObject"),
+		(ObjectInfo.IsSubObject ? TEXT("SubObject") : (ObjectInfo.IsDestructionInfo ? TEXT("DestructionInfo") : TEXT("RootObject"))),
 		LexToString(ObjectInfo.GetState()), 
 		SchedulingPriorities[ObjectIndex],
 		ObjectInfo.LastAckedBaselineIndex);
