Iris - Adhere to new UNetObjectPrioritizer API

#rnx
#Synced-CL 51619035

diff --git a/Engine/Source/Runtime/Net/Iris/Private/Iris/ReplicationSystem/Prioritization/LocationBasedNetObjectPrioritizer.cpp b/Engine/Source/Runtime/Net/Iris/Private/Iris/ReplicationSystem/Prioritization/LocationBasedNetObjectPrioritizer.cpp
--- a/Engine/Source/Runtime/Net/Iris/Private/Iris/ReplicationSystem/Prioritization/LocationBasedNetObjectPrioritizer.cpp
+++ b/Engine/Source/Runtime/Net/Iris/Private/Iris/ReplicationSystem/Prioritization/LocationBasedNetObjectPrioritizer.cpp
@@ -26,12 +26,12 @@ void ULocationBasedNetObjectPrioritizer::Deinit()
 	WorldLocations = nullptr;
 }
 
-void ULocationBasedNetObjectPrioritizer::OnMaxInternalNetRefIndexIncreased(uint32 NewMaxInternalIndex)
+void ULocationBasedNetObjectPrioritizer::OnMaxInternalNetRefIndexIncreased(UE::Net::FInternalNetRefIndex NewMaxInternalIndex)
 {
 	AssignedLocationIndices.SetNumBits(NewMaxInternalIndex);
 }
 
-bool ULocationBasedNetObjectPrioritizer::AddObject(uint32 ObjectIndex, FNetObjectPrioritizerAddObjectParams& Params)
+bool ULocationBasedNetObjectPrioritizer::AddObject(UE::Net::FInternalNetRefIndex ObjectIndex, FNetObjectPrioritizerAddObjectParams& Params)
 {
 	// We support either a world location in the state, tagged with RepTag_WorldLocation, or via the WorldLocations instance.
 	UE::Net::FRepTagFindInfo TagInfo;
@@ -65,7 +65,7 @@ bool ULocationBasedNetObjectPrioritizer::AddObject(uint32 ObjectIndex, FNetObjec
 	return true;
 }
 
-void ULocationBasedNetObjectPrioritizer::RemoveObject(uint32 ObjectIndex, const FNetObjectPrioritizationInfo& Info)
+void ULocationBasedNetObjectPrioritizer::RemoveObject(UE::Net::FInternalNetRefIndex ObjectIndex, const FNetObjectPrioritizationInfo& Info)
 {
 	const FObjectLocationInfo& ObjectInfo = static_cast<const FObjectLocationInfo&>(Info);
 	FreeLocation(ObjectInfo.GetLocationIndex());
@@ -111,7 +111,7 @@ void ULocationBasedNetObjectPrioritizer::SetLocation(const FObjectLocationInfo&
 	Locations[Info.GetLocationIndex()] = Location;
 }
 
-void ULocationBasedNetObjectPrioritizer::UpdateLocation(const uint32 ObjectIndex, const FObjectLocationInfo& Info, const UE::Net::FReplicationInstanceProtocol* InstanceProtocol)
+void ULocationBasedNetObjectPrioritizer::UpdateLocation(const UE::Net::FInternalNetRefIndex ObjectIndex, const FObjectLocationInfo& Info, const UE::Net::FReplicationInstanceProtocol* InstanceProtocol)
 {
 	if (Info.IsUsingWorldLocations())
 	{
diff --git a/Engine/Source/Runtime/Net/Iris/Private/Iris/ReplicationSystem/Prioritization/NetObjectCountLimiter.cpp b/Engine/Source/Runtime/Net/Iris/Private/Iris/ReplicationSystem/Prioritization/NetObjectCountLimiter.cpp
--- a/Engine/Source/Runtime/Net/Iris/Private/Iris/ReplicationSystem/Prioritization/NetObjectCountLimiter.cpp
+++ b/Engine/Source/Runtime/Net/Iris/Private/Iris/ReplicationSystem/Prioritization/NetObjectCountLimiter.cpp
@@ -53,7 +53,7 @@ void UNetObjectCountLimiter::RemoveConnection(uint32 ConnectionId)
 	Info = MoveTemp(Empty);
 }
 
-bool UNetObjectCountLimiter::AddObject(uint32 ObjectIndex, FNetObjectPrioritizerAddObjectParams& Params)
+bool UNetObjectCountLimiter::AddObject(UE::Net::FInternalNetRefIndex ObjectIndex, FNetObjectPrioritizerAddObjectParams& Params)
 {
 	FObjectInfo& Info = static_cast<FObjectInfo&>(Params.OutInfo);
 	const uint16 Index = AllocInternalIndex();
@@ -61,7 +61,7 @@ bool UNetObjectCountLimiter::AddObject(uint32 ObjectIndex, FNetObjectPrioritizer
 	return true;
 }
 
-void UNetObjectCountLimiter::RemoveObject(uint32 ObjectIndex, const FNetObjectPrioritizationInfo& InInfo)
+void UNetObjectCountLimiter::RemoveObject(UE::Net::FInternalNetRefIndex ObjectIndex, const FNetObjectPrioritizationInfo& InInfo)
 {
 	const FObjectInfo& Info = static_cast<const FObjectInfo&>(InInfo);
 	FreeInternalIndex(Info.GetPrioritizerInternalIndex());
diff --git a/Engine/Source/Runtime/Net/Iris/Private/Iris/ReplicationSystem/Prioritization/SphereWithOwnerBoostNetObjectPrioritizer.cpp b/Engine/Source/Runtime/Net/Iris/Private/Iris/ReplicationSystem/Prioritization/SphereWithOwnerBoostNetObjectPrioritizer.cpp
--- a/Engine/Source/Runtime/Net/Iris/Private/Iris/ReplicationSystem/Prioritization/SphereWithOwnerBoostNetObjectPrioritizer.cpp
+++ b/Engine/Source/Runtime/Net/Iris/Private/Iris/ReplicationSystem/Prioritization/SphereWithOwnerBoostNetObjectPrioritizer.cpp
@@ -34,7 +34,7 @@ void USphereWithOwnerBoostNetObjectPrioritizer::Deinit()
 	OwningConnections.Empty();
 }
 
-bool USphereWithOwnerBoostNetObjectPrioritizer::AddObject(uint32 ObjectIndex, FNetObjectPrioritizerAddObjectParams& Params)
+bool USphereWithOwnerBoostNetObjectPrioritizer::AddObject(UE::Net::FInternalNetRefIndex ObjectIndex, FNetObjectPrioritizerAddObjectParams& Params)
 {
 	if (!Super::AddObject(ObjectIndex, Params))
 	{
@@ -50,7 +50,7 @@ bool USphereWithOwnerBoostNetObjectPrioritizer::AddObject(uint32 ObjectIndex, FN
 	return true;
 }
 
-void USphereWithOwnerBoostNetObjectPrioritizer::RemoveObject(uint32 ObjectIndex, const FNetObjectPrioritizationInfo& Info)
+void USphereWithOwnerBoostNetObjectPrioritizer::RemoveObject(UE::Net::FInternalNetRefIndex ObjectIndex, const FNetObjectPrioritizationInfo& Info)
 {
 	const FObjectLocationInfo& ObjectInfo = static_cast<const FObjectLocationInfo&>(Info);
 	const uint32 OwningConnectionIndex = ObjectInfo.GetLocationIndex();
diff --git a/Engine/Source/Runtime/Net/Iris/Public/Iris/ReplicationSystem/Prioritization/LocationBasedNetObjectPrioritizer.h b/Engine/Source/Runtime/Net/Iris/Public/Iris/ReplicationSystem/Prioritization/LocationBasedNetObjectPrioritizer.h
--- a/Engine/Source/Runtime/Net/Iris/Public/Iris/ReplicationSystem/Prioritization/LocationBasedNetObjectPrioritizer.h
+++ b/Engine/Source/Runtime/Net/Iris/Public/Iris/ReplicationSystem/Prioritization/LocationBasedNetObjectPrioritizer.h
@@ -22,9 +22,9 @@ protected:
 	// UNetObjectPrioritizer interface
 	IRISCORE_API virtual void Init(FNetObjectPrioritizerInitParams& Params) override;
 	IRISCORE_API virtual void Deinit() override;
-	IRISCORE_API virtual void OnMaxInternalNetRefIndexIncreased(uint32 NewMaxInternalIndex) override;
-	IRISCORE_API virtual bool AddObject(uint32 ObjectIndex, FNetObjectPrioritizerAddObjectParams& Params) override;
-	IRISCORE_API virtual void RemoveObject(uint32 ObjectIndex, const FNetObjectPrioritizationInfo& Info) override;
+	IRISCORE_API virtual void OnMaxInternalNetRefIndexIncreased(UE::Net::FInternalNetRefIndex NewMaxInternalIndex) override;
+	IRISCORE_API virtual bool AddObject(UE::Net::FInternalNetRefIndex ObjectIndex, FNetObjectPrioritizerAddObjectParams& Params) override;
+	IRISCORE_API virtual void RemoveObject(UE::Net::FInternalNetRefIndex ObjectIndex, const FNetObjectPrioritizationInfo& Info) override;
 	IRISCORE_API virtual void UpdateObjects(FNetObjectPrioritizerUpdateParams&) override;
 
 protected:
@@ -47,7 +47,7 @@ protected:
 
 	IRISCORE_API VectorRegister GetLocation(const FObjectLocationInfo& Info) const;
 	IRISCORE_API void SetLocation(const FObjectLocationInfo& Info, VectorRegister Location);
-	IRISCORE_API void UpdateLocation(const uint32 ObjectIndex, const FObjectLocationInfo& Info, const UE::Net::FReplicationInstanceProtocol* InstanceProtocol);
+	IRISCORE_API void UpdateLocation(const UE::Net::FInternalNetRefIndex ObjectIndex, const FObjectLocationInfo& Info, const UE::Net::FReplicationInstanceProtocol* InstanceProtocol);
 
 private:
 	enum : unsigned
diff --git a/Engine/Source/Runtime/Net/Iris/Public/Iris/ReplicationSystem/Prioritization/NetObjectCountLimiter.h b/Engine/Source/Runtime/Net/Iris/Public/Iris/ReplicationSystem/Prioritization/NetObjectCountLimiter.h
--- a/Engine/Source/Runtime/Net/Iris/Public/Iris/ReplicationSystem/Prioritization/NetObjectCountLimiter.h
+++ b/Engine/Source/Runtime/Net/Iris/Public/Iris/ReplicationSystem/Prioritization/NetObjectCountLimiter.h
@@ -73,11 +73,11 @@ protected:
 	// UNetObjectPrioritizer interface
 	IRISCORE_API virtual void Init(FNetObjectPrioritizerInitParams&) override;
 	IRISCORE_API virtual void Deinit() override;
-	virtual void OnMaxInternalNetRefIndexIncreased(uint32 NewMaxInternalIndex) override {}
+	virtual void OnMaxInternalNetRefIndexIncreased(UE::Net::FInternalNetRefIndex NewMaxInternalIndex) override {}
 	IRISCORE_API virtual void AddConnection(uint32 ConnectionId) override;
 	IRISCORE_API virtual void RemoveConnection(uint32 ConnectionId) override;
-	IRISCORE_API virtual bool AddObject(uint32 ObjectIndex, FNetObjectPrioritizerAddObjectParams&) override;
-	IRISCORE_API virtual void RemoveObject(uint32 ObjectIndex, const FNetObjectPrioritizationInfo&) override;
+	IRISCORE_API virtual bool AddObject(UE::Net::FInternalNetRefIndex ObjectIndex, FNetObjectPrioritizerAddObjectParams&) override;
+	IRISCORE_API virtual void RemoveObject(UE::Net::FInternalNetRefIndex ObjectIndex, const FNetObjectPrioritizationInfo&) override;
 	IRISCORE_API virtual void UpdateObjects(FNetObjectPrioritizerUpdateParams&) override;
 	IRISCORE_API virtual void PrePrioritize(FNetObjectPrePrioritizationParams&) override;
 	IRISCORE_API virtual void Prioritize(FNetObjectPrioritizationParams&) override;
diff --git a/Engine/Source/Runtime/Net/Iris/Public/Iris/ReplicationSystem/Prioritization/SphereWithOwnerBoostNetObjectPrioritizer.h b/Engine/Source/Runtime/Net/Iris/Public/Iris/ReplicationSystem/Prioritization/SphereWithOwnerBoostNetObjectPrioritizer.h
--- a/Engine/Source/Runtime/Net/Iris/Public/Iris/ReplicationSystem/Prioritization/SphereWithOwnerBoostNetObjectPrioritizer.h
+++ b/Engine/Source/Runtime/Net/Iris/Public/Iris/ReplicationSystem/Prioritization/SphereWithOwnerBoostNetObjectPrioritizer.h
@@ -26,8 +26,8 @@ protected:
 	// UNetObjectPrioritizer interface
 	IRISCORE_API virtual void Init(FNetObjectPrioritizerInitParams& Params) override;
 	IRISCORE_API virtual void Deinit() override;
-	IRISCORE_API virtual bool AddObject(uint32 ObjectIndex, FNetObjectPrioritizerAddObjectParams& Params) override;
-	IRISCORE_API virtual void RemoveObject(uint32 ObjectIndex, const FNetObjectPrioritizationInfo& Info) override;
+	IRISCORE_API virtual bool AddObject(UE::Net::FInternalNetRefIndex ObjectIndex, FNetObjectPrioritizerAddObjectParams& Params) override;
+	IRISCORE_API virtual void RemoveObject(UE::Net::FInternalNetRefIndex ObjectIndex, const FNetObjectPrioritizationInfo& Info) override;
 	IRISCORE_API virtual void UpdateObjects(FNetObjectPrioritizerUpdateParams&) override;
 	IRISCORE_API virtual void Prioritize(FNetObjectPrioritizationParams&) override;
 
