NetTrace - Ignore too deeply nested traceevents.

#Synced-CL 51877528

diff --git a/Engine/Plugins/Runtime/ReplicationSystemTestPlugin/Source/Private/Tests/Serialization/TestNetBitStreams.cpp b/Engine/Plugins/Runtime/ReplicationSystemTestPlugin/Source/Private/Tests/Serialization/TestNetBitStreams.cpp
--- a/Engine/Plugins/Runtime/ReplicationSystemTestPlugin/Source/Private/Tests/Serialization/TestNetBitStreams.cpp
+++ b/Engine/Plugins/Runtime/ReplicationSystemTestPlugin/Source/Private/Tests/Serialization/TestNetBitStreams.cpp
@@ -5,6 +5,7 @@
 #include "Iris/Serialization/NetBitStreamReader.h"
 #include "Iris/Serialization/NetBitStreamWriter.h"
 #include "Iris/Serialization/NetBitStreamUtil.h"
+#include "Net/Core/Trace/NetTrace.h"
 
 namespace UE::Net::Private
 {
@@ -966,4 +967,33 @@ UE_NET_TEST_FIXTURE(FNetBitStreamReaderWriterTest, WriteBool)
 	}
 }
 
+#if UE_NET_TRACE_ENABLED
+static void WriteDataRecursive(FNetBitStreamWriter& Writer, FNetTraceCollector* Collector, uint32 CurrentDepth)
+{
+	UE_NET_TRACE_SCOPE(WrittenData, Writer, Collector, ENetTraceVerbosity::Verbose);
+	Writer.WriteBits(1U, 1);
+	if (CurrentDepth > 0)
+	{
+		WriteDataRecursive(Writer, Collector, CurrentDepth - 1U);
+	}
+}
+
+UE_NET_TEST_FIXTURE(FNetBitStreamWriterTest, TraceScopeClampedNestingLevel)
+{
+	FNetTraceCollector* Collector = UE_NET_TRACE_CREATE_COLLECTOR(ENetTraceVerbosity::Trace);
+
+	Writer.InitBytes(Buffer.GetBuffer(), Buffer.GetBufferCapacity());
+
+	{
+		UE_NET_TRACE_SCOPE(Payload, Writer, Collector, ENetTraceVerbosity::Verbose);
+		const uint32 TestNestedDepth = FNetTraceCollector::MaxNestingLevel + 5;
+	
+		// Just verify we do not trigger check/crash.
+		WriteDataRecursive(Writer, Collector, TestNestedDepth);
+	}
+
+	UE_NET_TRACE_DESTROY_COLLECTOR(Collector);
+}
+#endif
+
 }
diff --git a/Engine/Source/Runtime/Net/Core/Public/Net/Core/Trace/Private/NetTraceInternal.h b/Engine/Source/Runtime/Net/Core/Public/Net/Core/Trace/Private/NetTraceInternal.h
--- a/Engine/Source/Runtime/Net/Core/Public/Net/Core/Trace/Private/NetTraceInternal.h
+++ b/Engine/Source/Runtime/Net/Core/Public/Net/Core/Trace/Private/NetTraceInternal.h
@@ -376,6 +376,7 @@ public:
 	inline FNetTracePacketContentEvent& EditEvent(uint32 Index) { return Events[Index]; }
 	inline FNetTracePacketContentEvent& EditEventNoCheck(uint32 Index) { return Events.GetData()[Index]; }
 	inline void AddReferencedObjectId(uint64 ObjectId) { ReferencedObjectIds.Add(ObjectId); }
+	inline bool CanBeginPacketContentEvent() const { return CurrentNestingLevel < (MaxNestingLevel - 1); }
 	
 private:
 	friend struct FNetTrace;
@@ -524,7 +525,7 @@ FNetTraceEventScope<T>::FNetTraceEventScope(uint64 InObjectId, T& InStream, FNet
 	, Stream(InStream)
 	, EventIndex(FNetTrace::InvalidEventIndex)
 {
-	if (Collector)
+	if (Collector && Collector->CanBeginPacketContentEvent())
 	{
 		EventIndex = FNetTrace::BeginPacketContentEvent(*Collector, ENetTracePacketContentEventType::Object, GetBitStreamPositionForNetTrace(Stream));
 		Collector->EditEventNoCheck(EventIndex).ObjectId = InObjectId;
@@ -542,7 +543,7 @@ FNetTraceEventScope<T>::FNetTraceEventScope(TraceNameFunctor&& TraceNameFunc, T&
 , Stream(InStream)
 , EventIndex(FNetTrace::InvalidEventIndex)
 {
-	if (Collector)
+	if (Collector && Collector->CanBeginPacketContentEvent())
 	{
 		EventIndex = FNetTrace::BeginPacketContentEvent(*Collector, ENetTracePacketContentEventType::NameId, GetBitStreamPositionForNetTrace(Stream));
 		Collector->EditEventNoCheck(EventIndex).DebugNameId = TraceNameFunc();
