Iris - Fix issue with destroyed object with destruction info being mistaken for a destruction info. Allow destroyed object with destruction info to be replicated if it needs flushing for example due to a reliable RPC being called the same frame it's added for replication.

diff --git a/Engine/Source/Runtime/Net/Iris/Private/Iris/ReplicationSystem/NetRefHandleManager.cpp b/Engine/Source/Runtime/Net/Iris/Private/Iris/ReplicationSystem/NetRefHandleManager.cpp
--- a/Engine/Source/Runtime/Net/Iris/Private/Iris/ReplicationSystem/NetRefHandleManager.cpp
+++ b/Engine/Source/Runtime/Net/Iris/Private/Iris/ReplicationSystem/NetRefHandleManager.cpp
@@ -626,20 +626,21 @@ FNetRefHandle FNetRefHandleManager::CreateHandleForDestructionInfo(FNetRefHandle
 
 	if (DestructionInfoHandle.IsValid())
 	{
-		const uint32 InternalIndex = GetInternalIndex(DestructionInfoHandle);
+		const uint32 DestructionInfoInternalIndex = GetInternalIndex(DestructionInfoHandle);
 		const uint32 DestroyedInternalIndex = GetInternalIndex(Handle);
 		
-		// mark the internal index
-		DestroyedStartupObjectInternalIndices.SetBit(InternalIndex);
+		// Mark the internal index as a destruction info
+		GetReplicatedObjectDataNoCheck(DestructionInfoInternalIndex).bIsDestructionInfo = true;
+		DestroyedStartupObjectInternalIndices.SetBit(DestructionInfoInternalIndex);
 				
-		// Is the object replicated, if so we must make sure that we do not add it to scope by accident
+		// If the destroyed object is replicated we must make sure that we do not add it to scope by accident
 		if (DestroyedInternalIndex)
 		{			
-			DestroyedStartupObject.FindOrAdd(InternalIndex, DestroyedInternalIndex);
+			DestroyedStartupObject.FindOrAdd(DestructionInfoInternalIndex, DestroyedInternalIndex);
 		
 			// Mark the replicated index as destroyed
 			DestroyedStartupObjectInternalIndices.SetBit(DestroyedInternalIndex);
-			DestroyedStartupObject.FindOrAdd(DestroyedInternalIndex, InternalIndex);
+			DestroyedStartupObject.FindOrAdd(DestroyedInternalIndex, DestructionInfoInternalIndex);
 		}
 	}
 
diff --git a/Engine/Source/Runtime/Net/Iris/Private/Iris/ReplicationSystem/NetRefHandleManager.h b/Engine/Source/Runtime/Net/Iris/Private/Iris/ReplicationSystem/NetRefHandleManager.h
--- a/Engine/Source/Runtime/Net/Iris/Private/Iris/ReplicationSystem/NetRefHandleManager.h
+++ b/Engine/Source/Runtime/Net/Iris/Private/Iris/ReplicationSystem/NetRefHandleManager.h
@@ -135,6 +135,7 @@ public:
 		uint16 bHasCachedCreationInfo : 1U = 0;
 		uint16 bIsDebugObject : 1U = 0;
 		uint16 bNetObjectIsPendingDestroy : 1U = 0;
+		uint16 bIsDestructionInfo : 1U = 0;
 
 		/** The factory responsible for instantiating this object */
 		FNetObjectFactoryId NetFactoryId = InvalidNetObjectFactoryId;
@@ -416,7 +417,13 @@ public:
 
 	void AddReferencedObjects(FReferenceCollector& Collector);
 
+	/** IsDestructionInfo will return true if it's a destruction info for a destroyed startup object. */
+	inline bool IsDestructionInfo(FInternalNetRefIndex InternalIndex) const;
+
+	/** IsDestroyedStartupObject will return true if it's a destroyed startup object or the destruction info for a destroyed startup object. */
 	bool IsDestroyedStartupObject(FInternalNetRefIndex InternalIndex) const { return DestroyedStartupObjectInternalIndices.GetBit(InternalIndex); }
+
+	/** Will return the destruction info object index if passing a destroyed startup object index and vice versa. */
 	inline uint32 GetOriginalDestroyedStartupObjectIndex(FInternalNetRefIndex InternalIndex) const;
 
 	const FNetBitArrayView GetDestroyedStartupObjectInternalIndices() const { return MakeNetBitArrayView(DestroyedStartupObjectInternalIndices); }
@@ -747,6 +754,11 @@ bool FNetRefHandleManager::IsRootObject(FNetRefHandle Handle) const
 	return (InternalIndex != InvalidInternalNetRefIndex) ? IsRootObject(InternalIndex) : false;
 }
 
+bool FNetRefHandleManager::IsDestructionInfo(FInternalNetRefIndex InternalIndex) const
+{
+	return GetReplicatedObjectDataNoCheck(InternalIndex).bIsDestructionInfo;
+}
+
 uint32 FNetRefHandleManager::GetOriginalDestroyedStartupObjectIndex(FInternalNetRefIndex InternalIndex) const
 {
 	const uint32* FoundOriginalInternalIndex = DestroyedStartupObject.Find(InternalIndex);
diff --git a/Engine/Source/Runtime/Net/Iris/Private/Iris/ReplicationSystem/ReplicationWriter.cpp b/Engine/Source/Runtime/Net/Iris/Private/Iris/ReplicationSystem/ReplicationWriter.cpp
--- a/Engine/Source/Runtime/Net/Iris/Private/Iris/ReplicationSystem/ReplicationWriter.cpp
+++ b/Engine/Source/Runtime/Net/Iris/Private/Iris/ReplicationSystem/ReplicationWriter.cpp
@@ -550,19 +550,21 @@ void FReplicationWriter::StartReplication(uint32 InternalIndex)
 	Info.LastAckedBaselineIndex = FDeltaCompressionBaselineManager::InvalidBaselineIndex;
 	Info.PendingBaselineIndex = FDeltaCompressionBaselineManager::InvalidBaselineIndex;
 
-	const uint32 bIsDestructionInfo = NetRefHandleManager->IsDestroyedStartupObject(InternalIndex) ? 1U : 0U;	
+	const bool bIsDestructionInfo = NetRefHandleManager->IsDestructionInfo(InternalIndex);	
 	if (bIsDestructionInfo)
 	{
-		// Check status of original object about to be destroyed, if it has been confirmed as created we do not replicate the destruction info object at all
+		// Check status of original object about to be destroyed. Due to supporting sending reliable attachments for objects that have been created and destroyed in the same frame we may encounter a PendingCreate object which should be replicated.
 		if (const uint32 OriginalInternalIndex = NetRefHandleManager->GetOriginalDestroyedStartupObjectIndex(InternalIndex))
 		{
 			const FReplicationInfo& OriginalInfo = GetReplicationInfo(OriginalInternalIndex);
-			if ((OriginalInfo.GetState() != EReplicatedObjectState::Invalid && OriginalInfo.IsCreationConfirmed) && (OriginalInfo.GetState() != EReplicatedObjectState::WaitOnDestroyConfirmation))
+			const EReplicatedObjectState OriginalObjectState = OriginalInfo.GetState();
+			if ((OriginalObjectState != EReplicatedObjectState::Invalid) && 
+				((OriginalInfo.IsCreationConfirmed && (OriginalObjectState != EReplicatedObjectState::WaitOnDestroyConfirmation)) || ((OriginalObjectState == EReplicatedObjectState::PendingCreate) && (OriginalInfo.FlushFlags != EFlushFlags::FlushFlags_None))))
 			{
 				// We do not need to send the destruction info so we mark it as PermanentlyDestroyed
 				SetState(InternalIndex, EReplicatedObjectState::PermanentlyDestroyed);
 
-				Info.IsDestructionInfo = bIsDestructionInfo;
+				Info.IsDestructionInfo = 1U;
 				Info.IsCreationConfirmed = 1U;
 
 				NetRefHandleManager->AddNetObjectRef(InternalIndex);
