Iris
- Added option to objectreferrencecache to allow trying to async load of objects in partially gcd packages.
- Can be disabled with net.iris.AllowAsyncReload 0

#Synced-CL 51877528

diff --git a/Engine/Source/Runtime/Net/Iris/Private/Iris/ReplicationSystem/ObjectReferenceCache.cpp b/Engine/Source/Runtime/Net/Iris/Private/Iris/ReplicationSystem/ObjectReferenceCache.cpp
--- a/Engine/Source/Runtime/Net/Iris/Private/Iris/ReplicationSystem/ObjectReferenceCache.cpp
+++ b/Engine/Source/Runtime/Net/Iris/Private/Iris/ReplicationSystem/ObjectReferenceCache.cpp
@@ -36,6 +36,7 @@
 #include "Net/Core/Trace/NetTrace.h"
 
 #include "UObject/Package.h"
+#include "UObject/SoftObjectPath.h"
 
 #ifndef UE_NET_VALIDATE_REFERENCECACHE
 #	define UE_NET_VALIDATE_REFERENCECACHE 0
@@ -63,6 +64,9 @@ namespace UE::Net::Private
 static bool bIrisAllowAsyncLoading = true;
 FAutoConsoleVariableRef CVarIrisAllowAsyncLoading(TEXT("net.iris.AllowAsyncLoading"), bIrisAllowAsyncLoading, TEXT("Flag to allow or disallow async loading when using iris replication. Note: net.allowAsyncLoading must also be enabled."), ECVF_Default);
 
+static bool bIrisAllowAsyncReload = true;
+FAutoConsoleVariableRef CVarIrisAllowAsyncReload(TEXT("net.iris.AllowAsyncReload"), bIrisAllowAsyncReload, TEXT("Flag to allow async reload of objects from packages that might be partially garbage collected. Note: Only relevant if async loading is enabled."), ECVF_Default);
+
 #if UE_NET_ASYNCLOADING_DEBUG
 static bool bForceAsyncLoading = false;
 FAutoConsoleVariableRef CVarIrisForceAsyncLoading(TEXT("net.Iris.AsyncLoading.ForceEnable"), bForceAsyncLoading, TEXT("Force the async loading path even in Editor/PIE."), ECVF_Default);
@@ -905,7 +909,7 @@ UObject* FObjectReferenceCache::ResolveObjectReferenceHandleInternal(FNetRefHand
 		// Either the name should match, or this is dynamic, or we're on the server
 		// $IRIS: $TODO: Consider using FNames over string tokens, or store FNames in cache and fixup when resolving?
 #if UE_NET_VALIDATE_REFERENCECACHE
-		{			
+		{
 			//FString ObjectPath(StringTokenStore->ResolveRemoteToken(CacheObjectPtr->RelativePath, *ResolveContext.RemoteNetTokenStoreState));
 			//constexpr bool bReading = true;
 			//RenamePathForPie(ResolveContext.ConnectionId, ObjectPath, bReading, ResolveContext.ReplicationSystem);
@@ -1062,21 +1066,45 @@ UObject* FObjectReferenceCache::ResolveObjectReferenceHandleInternal(FNetRefHand
 			{
 				// Async loading disabled
 				Object = LoadPackage(nullptr, *ObjectPath, LOAD_None);
-				//SyncLoadedGUIDs.AddUnique(NetGUID);
 			}
 		}
 		else
 		{
-			// If we have a package, but for some reason didn't find the object then do a blocking load as a last attempt
+			// If we have a package, but for some reason didn't find the object we try to do a partial async load of the specific object or a blocking load as a last attempt
 			// This can happen for a few reasons:
 			//	1. The object was GC'd, but the package wasn't, so we need to reload
 			//	2. Someone else started async loading the outer package, and it's not fully loaded yet
-			Object = StaticLoadObject(UObject::StaticClass(), ObjOuter, *ObjectPath, nullptr, LOAD_NoWarn);
 
 			if (ShouldAsyncLoad() && !ResolveContext.bForceSyncLoad)
 			{
-				UE_LOG(LogIris, Error, TEXT( "GetObjectFromRefHandle: Forced blocking load. Path: %s, %s"), *ObjectPath, *RefHandle.ToString());
+				if (bIrisAllowAsyncReload)
+				{
+					// Construct full object path:
+					const bool bUseSubObjectDelimiter = ObjOuter->GetClass() != UPackage::StaticClass() && ObjOuter->GetOuter()->GetClass() == UPackage::StaticClass();
+					FSoftObjectPath SoftPath(FString::Printf(TEXT("%s%c%s"), *ObjOuter->GetPathName(), bUseSubObjectDelimiter ? SUBOBJECT_DELIMITER_CHAR : TEXT('.'), *ObjectPath));
+
+					if (!PendingAsyncObjectLoadRequests.Contains(SoftPath))
+					{
+						CacheObjectPtr->bIsPending = true;
+						StartAsyncLoadingObject(*CacheObjectPtr, SoftPath, ResolveContext.AsyncLoadingPriority, RefHandle, ResolveContext.ReplicationSystem);
+						UE_LOG_REFERENCECACHE(VeryVerbose, TEXT("GetObjectFromRefHandle: Async loading object. Path: %s, %s"), *SoftPath.ToString(), *RefHandle.ToString());
+					}
+					else
+					{
+						// Merge into existing request
+						PendingAsyncObjectLoadRequests[SoftPath].Merge(RefHandle);
+						CacheObjectPtr->bIsPending = true;
+						UE_LOG_REFERENCECACHE(VeryVerbose, TEXT("GetObjectFromRefHandle: Already async loading object. Path: %s, %s"), *SoftPath.ToString(), *RefHandle.ToString());
+					}
+					return nullptr;
+				}
+				else
+				{
+					UE_LOG(LogIris, Error, TEXT( "GetObjectFromRefHandle: Forced blocking load. Path: %s, %s"), *ObjectPath, *RefHandle.ToString());
+				}
 			}
+
+			Object = StaticLoadObject(UObject::StaticClass(), ObjOuter, *ObjectPath, nullptr, LOAD_NoWarn);
 		}
 	}
 
@@ -1117,7 +1145,6 @@ UObject* FObjectReferenceCache::ResolveObjectReferenceHandleInternal(FNetRefHand
 				// If package isn't fully loaded, load it now
 				UE_LOG_REFERENCECACHE(Warning, TEXT("ObjectReferenceCache::GetObjectFromRefHandle Blocking load of %s, %s"), ToCStr(ObjectPath), *RefHandle.ToString());
 				Object = LoadPackage(nullptr, ToCStr(ObjectPath), LOAD_None);
-				//SyncLoadedGUIDs.AddUnique(NetGUID);
 			}
 			else
 			{
@@ -1126,29 +1153,6 @@ UObject* FObjectReferenceCache::ResolveObjectReferenceHandleInternal(FNetRefHand
 		}
 	}
 
-	//if (CacheObjectPtr->NetworkChecksum != 0 && !CVarIgnoreNetworkChecksumMismatch.GetValueOnAnyThread())
-	//{
-	//	const uint32 NetworkChecksum = GetNetworkChecksum( Object );
-
-	//	if (CacheObjectPtr->NetworkChecksum != NetworkChecksum )
-	//	{
-	//		if ( NetworkChecksumMode == ENetworkChecksumMode::SaveAndUse )
-	//		{
-	//			FString ErrorStr = FString::Printf(TEXT("GetObjectFromRefHandle: Network checksum mismatch. ``IDPath: %s, %u, %u"), *FullRefHandlePath(RefHandle), CacheObjectPtr->NetworkChecksum, NetworkChecksum);
-	//			UE_LOG( LogNetPackageMap, Warning, TEXT("%s"), *ErrorStr );
-
-	//			CacheObjectPtr->bIsBroken = true;
-
-	//			BroadcastNetFailure(Driver, ENetworkFailure::NetChecksumMismatch, ErrorStr);
-	//			return NULL;
-	//		}
-	//		else
-	//		{
-	//			UE_LOG( LogNetPackageMap, Verbose, TEXT( "GetObjectFromRefHandle: Network checksum mismatch. ``IDPath: %s, %u, %u" ), *FullRefHandlePath( RefHandle ), CacheObjectPtr->NetworkChecksum, NetworkChecksum );
-	//		}
-	//	}
-	//}
-
 	UObjectReplicationBridge* ReplicationBridge = ResolveContext.ReplicationSystem->GetReplicationBridgeAs<UObjectReplicationBridge>();
 	if (Object && !ReplicationBridge->ObjectContainerHasFinishedLoading(Object))
 	{
@@ -1176,27 +1180,6 @@ UObject* FObjectReferenceCache::ResolveObjectReferenceHandleInternal(FNetRefHand
 	CacheObjectPtr->Object = MakeWeakObjectPtr(Object);
 	CacheObjectPtr->ObjectKey = Object;
 
-	// Assign the guid to the object 
-	// We don't want to assign this guid to the object if this guid is timing out
-	// But we'll have to if there is no other guid yet
-	//const bool bAllowClientRemap = !bIsNetGUIDAuthority && GbAllowClientRemapCacheObject;
-	//const bool bIsNotReadOnlyOrAllowRemap = (CacheObjectPtr->ReadOnlyTimestamp == 0 || bAllowClientRemap);
-
-	//if (bIsNotReadOnlyOrAllowRemap || !NetGUIDLookup.Contains(Object))
-	//{
-	//	if (CacheObjectPtr->ReadOnlyTimestamp > 0)
-	//	{
-	//		UE_LOG( LogNetPackageMap, Warning, TEXT( "GetObjectFromNetGUID: Attempt to reassign read-only guid. FullNetGUIDPath: %s" ), *FullNetGUIDPath( NetGUID ) );
-
-	//		if (bAllowClientRemap)
-	//		{
-	//			CacheObjectPtr->ReadOnlyTimestamp = 0;
-	//		}
-	//	}
-
-	//	NetGUIDLookup.Add( Object, NetGUID );
-	//}
-
 	// Promote to full handle
 	FNetRefHandle CompleteRefHandle = FNetRefHandleManager::MakeNetRefHandle(RefHandle.GetId(), ResolveContext.ReplicationSystem->GetId());
 
@@ -2276,7 +2259,7 @@ void FObjectReferenceCache::StartAsyncLoadingPackage(FCachedNetObjectReference&
 	FPendingAsyncLoadRequest* ExistingRequest = PendingAsyncLoadRequests.Find(PackagePath);
 	if (ExistingRequest)
 	{
-		// Same package name but a possibly different net GUID. Note down the GUID and wait for the async load completion callback
+		// Same package name but a possibly different NetRefHandle. Note down the NetRefHandle and wait for the async load completion callback
 		ExistingRequest->Merge(LoadRequest);
 		return;
 	}
@@ -2340,26 +2323,17 @@ void FObjectReferenceCache::AsyncPackageCallback(const FName& PackageName, UPack
 					UE_LOG_REFERENCECACHE(Error, TEXT("AsyncPackageCallback: Package FAILED to load. Path: %s, RefHandle: %s"), *PackageName.ToString(), *RefHandleToProcess.ToString());
 				}
 
+				// $TODO: Jira: UE-370863: Investigate if it is safe to resolve now to update CacheObject and mapping, or if it is best to wait until it is done as part of dispatching data.
+				// This does not really do anything since CachedObject->Object is not updated until the next resolve of the reference
 				if (UObject* Object = CacheObject->Object.Get())
 				{
 					// Something is loaded, we can now resolve queued data
 					UpdateTrackedQueuedBatchObjectReference(RefHandleToProcess, Object);
-
-					// $TODO: $IRIS
-					// This is odd, can we deprecate?, Does not really belong here IMHO
-					// If we need it, we can expose and bind a delegate from NetDriver
-					//if (UWorld* World = Object->GetWorld())
-					//{
-					//	if (AGameStateBase* GS = World->GetGameState())
-					//	{
-					//		GS->AsyncPackageLoaded(Object);
-					//	}
-					//}
 				}
 			}
 			else
 			{
-				UE_LOG_REFERENCECACHE(Error, TEXT("AsyncPackageCallback: Could not find net guid. Path: %s, RefHandle: %s"), *PackageName.ToString(), *RefHandleToProcess.ToString());
+				UE_LOG_REFERENCECACHE(Error, TEXT("AsyncPackageCallback: Could not find NetRefHandle. Path: %s, RefHandle: %s"), *PackageName.ToString(), *RefHandleToProcess.ToString());
 			}
 		}
 
@@ -2420,6 +2394,109 @@ void FObjectReferenceCache::AsyncPackageForcedFailCallback(const FName& PackageN
 #endif
 }
 
+void FObjectReferenceCache::StartAsyncLoadingObject(FCachedNetObjectReference& CacheObject, const FSoftObjectPath& ObjectPath, TAsyncLoadPriority AsyncLoadingPriority, const FNetRefHandle RefHandle, UReplicationSystem* ReplicationSystem)
+{
+	LLM_SCOPE_BYTAG(Iris);
+
+	const bool bIsPrioritySet = AsyncLoadingPriority != INDEX_NONE;
+	if (!bIsPrioritySet)
+	{
+		ensureMsgf(false, TEXT("AsyncLoadingPriority wasn't set in the ResolveContext. Using default value as fallback"));
+		AsyncLoadingPriority = ConvertAsyncLoadingPriority(EIrisAsyncLoadingPriority::Default);
+	}
+
+	FNetRefHandleManager& NetRefHandleManager = ReplicationSystem->GetReplicationSystemInternal()->GetNetRefHandleManager();
+
+	UE_LOG(LogIrisReferences, Verbose, TEXT("StartAsyncLoadingObject for %s: Object: %s (Async loading priority: %d)"), *NetRefHandleManager.PrintObjectFromNetRefHandle(RefHandle), *ObjectPath.ToString(), AsyncLoadingPriority);
+
+	FPendingAsyncObjectLoadRequest LoadRequest(RefHandle, ReplicationSystem->GetElapsedTime());
+
+	CacheObject.bIsPending = true;
+
+	FPendingAsyncObjectLoadRequest* ExistingRequest = PendingAsyncObjectLoadRequests.Find(ObjectPath);
+	if (ExistingRequest)
+	{
+		// Same object path but a possibly different NetRefHandle. Note down the NetRefHandle and wait for the async load completion callback
+		ExistingRequest->Merge(LoadRequest);
+		return;
+	}
+
+	PendingAsyncObjectLoadRequests.Emplace(ObjectPath, MoveTemp(LoadRequest));
+
+#if UE_NET_OBJECTREFERENCECACHE_ASYNCLOAD_DEBUG
+	FName PackageName = ObjectPath.GetAssetPath().GetPackageName();
+	if (bIrisAsyncLoadingFailAllLoads || bIrisAsyncLoadingFailNextLoad || (IrisAsyncLoadingFailPackageName.Len() > 0 && PackageName.ToString().Find(IrisAsyncLoadingFailPackageName) != INDEX_NONE))
+	{
+		bIrisAsyncLoadingFailNextLoad = false;
+		// Not thread safe. We don't want to load the object so we call our callback immediately with nullptr.
+		AsyncObjectCallback(ObjectPath, nullptr);
+	}
+	else
+#endif
+	{
+		FLoadAssetAsyncOptionalParams OptionalParams;
+		OptionalParams.PackagePriority = AsyncLoadingPriority;
+
+		ObjectPath.LoadAsync(
+			FLoadSoftObjectPathAsyncDelegate::CreateWeakLambda(ReplicationSystem,
+				[this, ObjectPath](const FSoftObjectPath& InObjectPath, UObject* LoadedObject)
+				{
+					AsyncObjectCallback(InObjectPath, LoadedObject);
+				}),
+			OptionalParams
+		);
+	}
+}
+
+void FObjectReferenceCache::AsyncObjectCallback(const FSoftObjectPath& ObjectPath, UObject* LoadedObject)
+{
+	LLM_SCOPE_BYTAG(Iris);
+
+	if (FPendingAsyncObjectLoadRequest const* const PendingLoadRequest = PendingAsyncObjectLoadRequests.Find(ObjectPath))
+	{
+		const bool bIsBroken = (LoadedObject == nullptr);
+
+		UE_IRIS_PARALLEL_EXPR(UE::TScopeLock InternalLock(ObjectReferenceCacheInternalCS));
+
+		for (FNetRefHandle RefHandleToProcess : PendingLoadRequest->NetRefHandles)
+		{
+			if (FCachedNetObjectReference* CacheObject = ReferenceHandleToCachedReference.Find(RefHandleToProcess))
+			{
+				if (!CacheObject->bIsPending)
+				{
+					UE_LOG_REFERENCECACHE(Error, TEXT("AsyncObjectCallback: Object wasn't pending. Path: %s, RefHandle: %s"), *ObjectPath.ToString(), *RefHandleToProcess.ToString());
+				}
+
+				CacheObject->bIsPending = false;
+
+				if (bIsBroken)
+				{
+					CacheObject->bIsBroken = true;
+					UE_LOG_REFERENCECACHE(Error, TEXT("AsyncObjectCallback: Object FAILED to load. Path: %s, RefHandle: %s"), *ObjectPath.ToString(), *RefHandleToProcess.ToString());
+				}
+
+				// $TODO: Jira: UE-370863: Investigate if it is safe to resolve now to update CacheObject and mapping, or if it is best to wait until it is done as part of dispatching data.
+
+				if (UObject* Object = LoadedObject)
+				{
+					// Something is loaded, we can now resolve queued data
+					UpdateTrackedQueuedBatchObjectReference(RefHandleToProcess, Object);
+				}
+			}
+			else
+			{
+				UE_LOG_REFERENCECACHE(Error, TEXT("AsyncObjectCallback: Could not find net ref handle. Path: %s, RefHandle: %s"), *ObjectPath.ToString(), *RefHandleToProcess.ToString());
+			}
+		}
+
+		PendingAsyncObjectLoadRequests.Remove(ObjectPath);
+	}
+	else
+	{
+		UE_LOG_REFERENCECACHE(Error, TEXT("AsyncObjectCallback: Could not find object. Path: %s"), *ObjectPath.ToString());
+	}
+}
+
 void FObjectReferenceCache::SetAsyncLoadMode(const EAsyncLoadMode NewMode)
 {
 	AsyncLoadMode = NewMode;
diff --git a/Engine/Source/Runtime/Net/Iris/Private/Iris/ReplicationSystem/ObjectReferenceCache.h b/Engine/Source/Runtime/Net/Iris/Private/Iris/ReplicationSystem/ObjectReferenceCache.h
--- a/Engine/Source/Runtime/Net/Iris/Private/Iris/ReplicationSystem/ObjectReferenceCache.h
+++ b/Engine/Source/Runtime/Net/Iris/Private/Iris/ReplicationSystem/ObjectReferenceCache.h
@@ -11,6 +11,7 @@
 #include "Containers/Map.h"
 #include "ObjectReferenceCacheFwd.h"
 #include "UObject/ObjectPtr.h"
+#include "UObject/SoftObjectPath.h"
 
 enum class EIrisAsyncLoadingPriority : uint8;
 
@@ -229,11 +230,13 @@ private:
 		void Merge(const FPendingAsyncLoadRequest& Other);
 		void Merge(FNetRefHandle InNetRefHandle);
 
-		// NetRefHandles that requested loading for the same UPackage
+		// NetRefHandles that requested loading for the same UPackage/UObject
 		TArray<FNetRefHandle, TInlineAllocator<4>> NetRefHandles;
 		double RequestStartTime;
 	};
 
+	using FPendingAsyncObjectLoadRequest = FPendingAsyncLoadRequest;
+
 	bool CreateObjectReferenceInternal(const UObject* Object, FNetObjectReference& OutReference, UReplicationSystem* ReplicationSystem);
 
 
@@ -268,6 +271,9 @@ private:
 	void StartAsyncLoadingPackage(FCachedNetObjectReference& Object, FName PackagePath, TAsyncLoadPriority AsyncLoadingPriority, const FNetRefHandle RefHandle, const bool bWasAlreadyAsyncLoading, UReplicationSystem* ReplicationSystem);
 	void ValidateAsyncLoadingPackage(FCachedNetObjectReference& Object, FName PackagePath, const FNetRefHandle RefHandle);
 
+	void StartAsyncLoadingObject(FCachedNetObjectReference& CacheObject, const FSoftObjectPath& ObjectPath, TAsyncLoadPriority AsyncLoadingPriority, const FNetRefHandle RefHandle, UReplicationSystem* ReplicationSystem);
+	void AsyncObjectCallback(const FSoftObjectPath& ObjectPath, UObject* LoadedObject);
+
 private:
 
 #if UE_SUPPORT_PARALLEL_IRIS
@@ -305,6 +311,9 @@ private:
 	/** Set of packages that are currently pending async loads, referenced by package name. */
 	TMap<FName, FPendingAsyncLoadRequest> PendingAsyncLoadRequests;
 
+	/** Set of objects (non-package) that are currently pending async loads, referenced by soft object path. */
+	TMap<FSoftObjectPath, FPendingAsyncObjectLoadRequest> PendingAsyncObjectLoadRequests;
+
 	/** Optional delegates that can be used to override the UObject<->NetRefHandle id assignments */
 	FObjectToNetRefHandleSerial ObjectToNetRefHandleSerial;
 	FNetRefHandleSerialToObject NetRefHandleSerialToObject;
