Iris - Removed bad early return when outputting default state of replication state descriptors, for example when encountering protocol mismatch errors.

#Synced-CL 51996830

diff --git a/Engine/Source/Runtime/Net/Iris/Private/Iris/ReplicationSystem/ReplicationOperations.cpp b/Engine/Source/Runtime/Net/Iris/Private/Iris/ReplicationSystem/ReplicationOperations.cpp
--- a/Engine/Source/Runtime/Net/Iris/Private/Iris/ReplicationSystem/ReplicationOperations.cpp
+++ b/Engine/Source/Runtime/Net/Iris/Private/Iris/ReplicationSystem/ReplicationOperations.cpp
@@ -72,7 +72,7 @@ static bool bDeltaCompressInitialState = true;
 static FAutoConsoleVariableRef CVarbDeltaCompressInitialState(
 	TEXT("net.Iris.DeltaCompressInitialState"),
 	bDeltaCompressInitialState,
-	TEXT("if true we compare with default state when serializing inital state."
+	TEXT("if true we compare with default state when serializing initial state."
 	));
 
 static bool bOnlyQuantizeDirtyMembers = true;
@@ -283,11 +283,7 @@ static void OutputDescriptorMemberHashToString(FNetSerializationContext Context,
 
 	auto AppendMemberAndIndentation = [](FStringBuilderBase& StringBuilder, uint32 MemberIndex, uint32 Indentation)
 	{
-		while (Indentation--)
-		{
-			StringBuilder.Append(TEXT("    "));
-		};
-		StringBuilder.Appendf(TEXT("Member: %u "), MemberIndex);
+		StringBuilder.Appendf(TEXT("%*sMember: %u "), 4*Indentation, TEXT(""), MemberIndex);
 	};
 
 	for (uint32 MemberIt = 0; MemberIt < MemberCount; ++MemberIt)
@@ -305,7 +301,6 @@ static void OutputDescriptorMemberHashToString(FNetSerializationContext Context,
 
 			StringBuilder.Appendf(TEXT("%s\n"), ToCStr(MemberDebugDescriptors[MemberIt].DebugName));
 			OutputDescriptorMemberHashToString(Context, HashBuffer, Indentation + 1, StringBuilder, SrcInternalBuffer + MemberDescriptor.InternalMemberOffset, StructDescriptor);
-			return;	
 		}
 		else
 		{
