Iris - Close config system prior to static shutdown.

#rnx
#Synced-CL 51996830

diff --git a/Engine/Source/Runtime/Net/Iris/Private/Iris/Config/IrisDynamicConfig.cpp b/Engine/Source/Runtime/Net/Iris/Private/Iris/Config/IrisDynamicConfig.cpp
--- a/Engine/Source/Runtime/Net/Iris/Private/Iris/Config/IrisDynamicConfig.cpp
+++ b/Engine/Source/Runtime/Net/Iris/Private/Iris/Config/IrisDynamicConfig.cpp
@@ -6,6 +6,7 @@
 #include "Containers/Map.h"
 #include "Iris/Core/IrisLog.h"
 #include "Misc/ConfigCacheIni.h"
+#include "Misc/CoreDelegates.h"
 
 namespace UE::Net::Private
 {
@@ -24,6 +25,7 @@ public:
 	FIrisDynamicConfig::FOnIrisDynamicConfigChange OnIrisDynamicConfigChange;
 
 private:
+	void OnAppPreExit();
 	FConfigBranch* GetOrCreateConfig();
 
 	// The config system uses three different strings types to retrieve things...
@@ -83,12 +85,20 @@ void FIrisDynamicConfig::GetArray(const TCHAR* SectionName, const TCHAR* ArrayNa
 namespace UE::Net::Private
 {
 
+void FIrisDynamicConfigImpl::OnAppPreExit()
+{
+	IrisConfig = nullptr;
+	ConfigCache.Reset();
+}
+
 FConfigBranch* FIrisDynamicConfigImpl::GetOrCreateConfig()
 {
 	if (IrisConfig == nullptr)
 	{
 		ConfigCache.Emplace(EConfigCacheType::Temporary);
 		IrisConfig = &ConfigCache->AddNewBranch(VirtualBranchFileName);
+
+		FCoreDelegates::OnPreExit.AddRaw(this, &FIrisDynamicConfigImpl::OnAppPreExit);
 	}
 
 	UE_ASSUME(IrisConfig != nullptr);
@@ -97,6 +107,11 @@ FConfigBranch* FIrisDynamicConfigImpl::GetOrCreateConfig()
 
 void FIrisDynamicConfigImpl::RegisterDynamicConfig(FName UniqueConfigName, const FString& FilePath)
 {
+	if (IsEngineExitRequested())
+	{
+		return;
+	}
+
 	if (!ensure(IsInGameThread()))
 	{
 		UE_LOGF(LogIris, Error, "FIrisDynamicConfig::RegisterDynamicConfig called from non-game thread. IGNORING dynamic config named '%ls'.", ToCStr(UniqueConfigName.ToString()));
@@ -121,6 +136,11 @@ void FIrisDynamicConfigImpl::RegisterDynamicConfig(FName UniqueConfigName, const
 
 void FIrisDynamicConfigImpl::RegisterDynamicConfigBuffer(FName UniqueConfigName, const FString& VirtualFilePath, const FString& ConfigContents)
 {
+	if (IsEngineExitRequested())
+	{
+		return;
+	}
+
 	if (!ensure(IsInGameThread()))
 	{
 		UE_LOGF(LogIris, Error, "FIrisDynamicConfig::RegisterDynamicConfigBuffer called from non-game thread. IGNORING dynamic config named '%ls'.", ToCStr(UniqueConfigName.ToString()));
