Iris
- Fixed issues with CQTestframework failures due to some rooted dynamic references was treated as stale.

diff --git a/Engine/Source/Runtime/Net/Iris/Private/Iris/ReplicationSystem/ObjectReferenceCache.cpp b/Engine/Source/Runtime/Net/Iris/Private/Iris/ReplicationSystem/ObjectReferenceCache.cpp
--- a/Engine/Source/Runtime/Net/Iris/Private/Iris/ReplicationSystem/ObjectReferenceCache.cpp
+++ b/Engine/Source/Runtime/Net/Iris/Private/Iris/ReplicationSystem/ObjectReferenceCache.cpp
@@ -379,7 +379,8 @@ bool FObjectReferenceCache::CreateObjectReferenceInternal(const UObject* Object,
 						}
 						else
 						{
-							bOuterChainValid = false;
+							// Invalidate reference if we expected the outer to be valid.
+							bOuterChainValid = !CurrentCachedOuterHandle.IsValid();
 							break;
 						}
 					}
