Use TFunctionRefs instead of TFunction when choosing between different lambdas. It produces much leaner code.

#rnx
#Synced-CL 51996830

diff --git a/Engine/Source/Runtime/Net/Iris/Private/Iris/ReplicationSystem/Polling/ObjectPoller.cpp b/Engine/Source/Runtime/Net/Iris/Private/Iris/ReplicationSystem/Polling/ObjectPoller.cpp
--- a/Engine/Source/Runtime/Net/Iris/Private/Iris/ReplicationSystem/Polling/ObjectPoller.cpp
+++ b/Engine/Source/Runtime/Net/Iris/Private/Iris/ReplicationSystem/Polling/ObjectPoller.cpp
@@ -106,22 +106,20 @@ namespace CVars
 			FNetStatsContext* ChildNetStatsContext = PollTaskData.NetTypeStats->AcquireChildNetStatsContext();
 
 			//Set up the lambda to use for polling, selecting a different one for PushModel or ForcePoll
-			TFunction<void(uint32 ObjectIndex)> PollingLambdaFunc = PollTaskData.bUsePushModel ?
-				TFunction<void(uint32 ObjectIndex)>
+			auto PushModelPoll = [this, ChildNetStatsContext](uint32 ObjectIndex)
 				{
-					[this, ChildNetStatsContext](uint32 ObjectIndex)
-					{
-						Poller->PushModelPollObject(ObjectIndex, ChildNetStatsContext, PollTaskData.PollStats);
-					}
-				} :
-				TFunction<void(uint32 ObjectIndex)>
+					Poller->PushModelPollObject(ObjectIndex, ChildNetStatsContext, PollTaskData.PollStats);
+				};
+			auto ForcePoll = [this, ChildNetStatsContext](uint32 ObjectIndex)
 				{
-					[this, ChildNetStatsContext](uint32 ObjectIndex)
-					{
-						Poller->ForcePollObject(ObjectIndex, ChildNetStatsContext, PollTaskData.PollStats);
-					}
+					Poller->ForcePollObject(ObjectIndex, ChildNetStatsContext, PollTaskData.PollStats);
 				};
 
+			// Named TFunctionRefs needed to give the ternary a common type
+			TFunctionRef<void(uint32 ObjectIndex)> PushModelPollFunc = PushModelPoll;
+			TFunctionRef<void(uint32 ObjectIndex)> ForcePollFunc = ForcePoll;
+			TFunctionRef<void(uint32 ObjectIndex)> PollingLambdaFunc = PollTaskData.bUsePushModel ? PushModelPollFunc : ForcePollFunc;
+
 			//   For each Task, we want to take the objects list and process it based on
 			//Chunks of a specific size, and interleave these Chunks between tasks
 			//   This is more balanced than dividing into contiguous chunks as the object bit array
