Iris
- Changed allocation for ScheduledObjects to be growable instead of always allocate the max.

#Synced-CL 52054134

diff --git a/Engine/Source/Runtime/Net/Iris/Private/Iris/ReplicationSystem/ReplicationWriter.cpp b/Engine/Source/Runtime/Net/Iris/Private/Iris/ReplicationSystem/ReplicationWriter.cpp
--- a/Engine/Source/Runtime/Net/Iris/Private/Iris/ReplicationSystem/ReplicationWriter.cpp
+++ b/Engine/Source/Runtime/Net/Iris/Private/Iris/ReplicationSystem/ReplicationWriter.cpp
@@ -1161,7 +1161,7 @@ void FReplicationWriter::UpdatePriorities(const float* UpdatedPriorities)
 	ObjectsWithDirtyChanges.ForAllSetBits(UpdatePriority);
 }
 
-void FReplicationWriter::ScheduleDependentObjects(uint32 Index, float ParentPriority, TArray<float>& LocalPriorities, FScheduleObjectInfo* ScheduledObjectIndices, uint32& OutScheduledObjectCount)
+void FReplicationWriter::ScheduleDependentObjects(uint32 Index, float ParentPriority, TArray<float>& LocalPriorities, FScheduledObjectInfosArray& ScheduledObjectIndices)
 {
 	const float DependentObjectPriorityBump = UE_KINDA_SMALL_NUMBER;
 
@@ -1183,10 +1183,7 @@ void FReplicationWriter::ScheduleDependentObjects(uint32 Index, float ParentPrio
 				LocalPriorities[DependentInternalIndex] = UpdatedPriority;
 
 				// Schedule it, it does not matter if we add it to the scheduled list multiple times
-				FScheduleObjectInfo& ScheduledObjectInfo = ScheduledObjectIndices[OutScheduledObjectCount];
-				ScheduledObjectInfo.Index = DependentInternalIndex;
-				ScheduledObjectInfo.SortKey = UpdatedPriority;
-				++OutScheduledObjectCount;								
+				ScheduledObjectIndices.Add({.Index = DependentInternalIndex, .SortKey = UpdatedPriority});
 			}
 		}
 		
@@ -1194,18 +1191,16 @@ void FReplicationWriter::ScheduleDependentObjects(uint32 Index, float ParentPrio
 		// the scheduling order hint at least in relation to the parent, but a dependent object might also end up replicating before its parent`s parent
 		if (NetRefHandleManager->GetObjectsWithDependentObjectsInternalIndices().GetBit(DependentInternalIndex))
 		{
-			ScheduleDependentObjects(DependentInternalIndex, UpdatedPriority, LocalPriorities, ScheduledObjectIndices, OutScheduledObjectCount);
+			ScheduleDependentObjects(DependentInternalIndex, UpdatedPriority, LocalPriorities, ScheduledObjectIndices);
 		}
 	}
 }
 
-uint32 FReplicationWriter::ScheduleObjects(FScheduleObjectInfo* OutScheduledObjectIndices)
+uint32 FReplicationWriter::ScheduleObjects(FScheduledObjectInfosArray& OutScheduledObjectIndices)
 {
 	IRIS_PROFILER_SCOPE(FReplicationWriter_ScheduleObjects);
 
-	uint32 ScheduledObjectCount = 0;
-
-	FScheduleObjectInfo* ScheduledObjectIndices = OutScheduledObjectIndices;
+	FScheduledObjectInfosArray& ScheduledObjectIndices = OutScheduledObjectIndices;
 
 	// Special index is handled later.
 	ObjectsWithDirtyChanges.ClearBit(ObjectIndexForOOBAttachment);
@@ -1213,22 +1208,22 @@ uint32 FReplicationWriter::ScheduleObjects(FScheduleObjectInfo* OutScheduledObje
 	const FNetBitArray& UpdatedObjects = ObjectsWithDirtyChanges;
 	const FNetBitArray& SubObjects = NetRefHandleManager->GetSubObjectInternalIndices();
 
-	auto FillIndexListFunc = [&ScheduledObjectIndices, &ScheduledObjectCount, this](uint32 Index)
+	auto FillIndexListFunc = [&ScheduledObjectIndices, this](uint32 Index)
 	{
 		const float UpdatedPriority = SchedulingPriorities[Index];
 
-		FScheduleObjectInfo& ScheduledObjectInfo = ScheduledObjectIndices[ScheduledObjectCount];
+		FScheduleObjectInfo ScheduledObjectInfo = {.Index = Index,.SortKey = UpdatedPriority };
 		ScheduledObjectInfo.Index = Index;
 		ScheduledObjectInfo.SortKey = UpdatedPriority;
 
 		if (UpdatedPriority >= FReplicationWriter::SchedulingThresholdPriority)
 		{
-			++ScheduledObjectCount;
+			ScheduledObjectIndices.Add({.Index = Index,.SortKey = UpdatedPriority });
 
 			// If we have dependent objects that needs to replicate before parent we need to schedule them as well.
 			if (NetRefHandleManager->GetObjectsWithDependentObjectsInternalIndices().GetBit(Index))
 			{
-				ScheduleDependentObjects(Index, UpdatedPriority, SchedulingPriorities, ScheduledObjectIndices, ScheduledObjectCount);
+				ScheduleDependentObjects(Index, UpdatedPriority, SchedulingPriorities, ScheduledObjectIndices);
 			}
 		}
 	};
@@ -1237,7 +1232,7 @@ uint32 FReplicationWriter::ScheduleObjects(FScheduleObjectInfo* OutScheduledObje
 	FNetBitArray::ForAllSetBits(UpdatedObjects, SubObjects, FNetBitArray::AndNotOp, FillIndexListFunc);
 
 	// we now have our list of objects to write.
-	return ScheduledObjectCount;
+	return ScheduledObjectIndices.Num();
 }
 
 uint32 FReplicationWriter::SortScheduledObjects(FScheduleObjectInfo* ScheduledObjectIndices, uint32 ScheduledObjectCount, uint32 StartIndex)
@@ -3779,7 +3774,7 @@ uint32 FReplicationWriter::WriteObjects(FNetSerializationContext& Context)
 	uint32 WrittenObjectCount = 0;
 	
 	const uint32 ObjectCount = WriteContext.ScheduledObjectCount;
-	FScheduleObjectInfo* ObjectList = WriteContext.ScheduledObjectInfos;
+	FScheduleObjectInfo* ObjectList = WriteContext.ScheduledObjectInfos.GetData();
 
 	uint32 ObjectListIt = WriteContext.CurrentIndex;
 	uint32 SortedCount = WriteContext.SortedObjectCount;
@@ -4241,11 +4236,6 @@ UDataStream::EWriteResult FReplicationWriter::BeginWrite(const UDataStream::FBeg
 		WriteContext.bHasHugeObjectToSend = bHasUnsentHugeObject;
 		WriteContext.bHasOOBAttachmentsToSend = bHasUnsentOOBAttachments;
 
-		// $IRIS TODO: LinearAllocator/ScratchPad?
-		// Allocate space for indices to send
-		// This should be allocated from frame temp allocator and be cleaned up end of frame, we might want this data to persist over multiple write calls but not over multiple frames 
-		// https://jira.it.epicgames.com/browse/UE-127374	
-		WriteContext.ScheduledObjectInfos = reinterpret_cast<FScheduleObjectInfo*>(FMemory::Malloc(sizeof(FScheduleObjectInfo) * NetRefHandleManager->GetCurrentMaxInternalNetRefIndex()));
 		WriteContext.ScheduledObjectCount = ScheduleObjects(WriteContext.ScheduledObjectInfos);
 	}
 
@@ -4329,8 +4319,8 @@ void FReplicationWriter::EndWrite()
 		}
 #endif
 
-		FMemory::Free(WriteContext.ScheduledObjectInfos);
-		WriteContext.ScheduledObjectInfos = nullptr;
+		// Intentional use of Empty to avoid wasting memory due to rare frames with huge amounts of scheduled objects.
+		WriteContext.ScheduledObjectInfos.Empty();
 		WriteContext.bIsValid = false;
 	}
 }
diff --git a/Engine/Source/Runtime/Net/Iris/Private/Iris/ReplicationSystem/ReplicationWriter.h b/Engine/Source/Runtime/Net/Iris/Private/Iris/ReplicationSystem/ReplicationWriter.h
--- a/Engine/Source/Runtime/Net/Iris/Private/Iris/ReplicationSystem/ReplicationWriter.h
+++ b/Engine/Source/Runtime/Net/Iris/Private/Iris/ReplicationSystem/ReplicationWriter.h
@@ -68,6 +68,8 @@ public:
 	static constexpr float LostStatePriorityBump = 1.f;
 	static constexpr float SchedulingThresholdPriority = 1.f;
 
+	static constexpr uint32 PreAllocatedScheduledObjectCount = 2048U;
+
 	// When scheduling there is no point in scheduling more objects than we can fit in a packet
 	static const uint32 PartialSortObjectCount = 128u;
 
@@ -239,12 +241,14 @@ private:
 		float SortKey;
 	};
 
+	using FScheduledObjectInfosArray = TArray<FScheduleObjectInfo, TInlineAllocator<PreAllocatedScheduledObjectCount>>;
+
 	struct FExtraScheduleInfo
 	{
 		FExtraScheduleInfo(FInternalNetRefIndex InInternalIndex, bool bInAllowAppendCreationDependencies = true)
 		: InternalIndex(InInternalIndex) 
 		, bAppendCreationDependencies(bInAllowAppendCreationDependencies)
-		{			
+		{
 		}
 
 		FInternalNetRefIndex InternalIndex = 0U;
@@ -275,7 +279,7 @@ private:
 		TArray<FInternalNetRefIndex, TInlineAllocator<32>> ObjectsSkippedInThisPacket;
 
 		// Scheduled objects
-		FScheduleObjectInfo* ScheduledObjectInfos;
+		FScheduledObjectInfosArray ScheduledObjectInfos;
 		uint32 ScheduledObjectCount;
 
 		// For performance sake we do partial sorting so we need to track how many object we have sorted.
@@ -521,9 +525,9 @@ private:
 
 	void CommitBatchRecord(const FBatchRecord& BatchRecord);
 
-	void ScheduleDependentObjects(uint32 Index, float ParentPriority, TArray<float>& LocalPriorities, FScheduleObjectInfo* ScheduledObjectIndices, uint32& OutScheduledObjectCount);
+	void ScheduleDependentObjects(uint32 Index, float ParentPriority, TArray<float>& LocalPriorities, FScheduledObjectInfosArray& ScheduledObjectIndices);
 
-	uint32 ScheduleObjects(FScheduleObjectInfo* ScheduledObjectIndices);
+	uint32 ScheduleObjects(FScheduledObjectInfosArray& ScheduledObjectIndices);
 	
 	// Partial sort of OutScheduledObjectIndices, will sort at most PartialSortObjectCount objects
 	uint32 SortScheduledObjects(FScheduleObjectInfo* ScheduledObjectIndices, uint32 ScheduledObjectCount, uint32 StartIndex);
