Iris - Store poll frequencies in FObjectPollFrequencyLimiter to avoid having to call FNetObjectFactory::GetPollFrequency() when the application update frequency changes.
Deprecated FNetObjectFactory::GetPollFrequency().

#Synced-CL 51996830

#Synced-CL 51996830

diff --git a/Engine/Plugins/Runtime/ReplicationSystemTestPlugin/Source/Private/Tests/ReplicationSystem/TestObjectPollFrequencyLimiter.cpp b/Engine/Plugins/Runtime/ReplicationSystemTestPlugin/Source/Private/Tests/ReplicationSystem/TestObjectPollFrequencyLimiter.cpp
--- a/Engine/Plugins/Runtime/ReplicationSystemTestPlugin/Source/Private/Tests/ReplicationSystem/TestObjectPollFrequencyLimiter.cpp
+++ b/Engine/Plugins/Runtime/ReplicationSystemTestPlugin/Source/Private/Tests/ReplicationSystem/TestObjectPollFrequencyLimiter.cpp
@@ -15,25 +15,30 @@ namespace UE::Net::Private
 
 UE_NET_TEST(ObjectPollFrequencyLimiter, ObjectsToPollAreSpreadOutWhenCreatingObjectAndUpdatingOncePerFrame)
 {
-	constexpr uint32 PollFrequency = 17U;
-	constexpr uint32 MaxObjectCount = PollFrequency + 1U;
+	constexpr float UpdateFrequency = 1.0f;
+	constexpr uint32 PollFramesBetweenUpdates = 17U;
+	constexpr FInternalNetRefIndex MaxObjectCount = PollFramesBetweenUpdates + 1;
+	constexpr FInternalNetRefIndex MaxInternalIndex = MaxObjectCount + 1;
 
 	FObjectPollFrequencyLimiter FrequencyLimiter;
-	FrequencyLimiter.Init(MaxObjectCount + 1);
+	FrequencyLimiter.Init(MaxInternalIndex);
+	FrequencyLimiter.SetUpdateFrequency(UpdateFrequency);
 
-	FNetBitArray EmptyBitArray(MaxObjectCount + 1);
-	FNetBitArray ScopeBitArray(MaxObjectCount + 1);
+	FNetBitArray EmptyBitArray(MaxInternalIndex);
+	FNetBitArray ScopeBitArray(MaxInternalIndex);
 
 	// Create one object at a time and update.
-	FNetBitArray OutputArray(MaxObjectCount + 1);
+	FNetBitArray OutputArray(MaxInternalIndex);
 	FNetBitArrayView OutputBitArrayView = MakeNetBitArrayView(OutputArray);
 
 	for (uint32 It = 0, EndIt = MaxObjectCount; It < EndIt; ++It)
 	{
 		FInternalNetRefIndex ObjectIndex = It + 1U;
-		FrequencyLimiter.SetPollFramePeriod(ObjectIndex, PollFrequency);
+		FrequencyLimiter.SetPollFrequency(ObjectIndex, UpdateFrequency/(PollFramesBetweenUpdates + 1));
+
+		// Add object to the scope
+		ScopeBitArray.SetBit(ObjectIndex);
 
-		ScopeBitArray.SetBits(1, ObjectIndex);
 		FrequencyLimiter.Update(MakeNetBitArrayView(ScopeBitArray), MakeNetBitArrayView(EmptyBitArray), OutputBitArrayView);
 	}
 
@@ -54,20 +59,24 @@ UE_NET_TEST(ObjectPollFrequencyLimiter, ObjectsToPollAreSpreadOutRegardlessOfUpd
 	constexpr int32 RandSeed = 1192314641;
 	FMath::RandInit(RandSeed);
 
-	constexpr uint32 MaxObjectCount = 128;
+	constexpr float UpdateFrequency = 1.0f;
+	constexpr uint32 PollFramesBetweenUpdates = 8U;
+	constexpr FInternalNetRefIndex MaxObjectCount = 128;
+	constexpr FInternalNetRefIndex MaxInternalIndex = MaxObjectCount + 1;
+
 	FObjectPollFrequencyLimiter FrequencyLimiter;
-	FrequencyLimiter.Init(MaxObjectCount + 1);
+	FrequencyLimiter.Init(MaxInternalIndex);
+	FrequencyLimiter.SetUpdateFrequency(UpdateFrequency);
 
-	constexpr uint32 PollFrequency = 8U;
-	// With three times as many objects as the poll frequency we expect three objects to be polled per frame.
-	constexpr uint32 PollFrequencyObjectCount = 4U*(PollFrequency + 1U);
+	// With four times as many objects as the poll frequency we expect four objects to be polled per frame.
+	constexpr uint32 PollFrequencyObjectCount = 4U*(PollFramesBetweenUpdates + 1U);
 	constexpr uint32 MaxUpdateCountPerIteration = 5U;
 
-	FNetBitArray EmptyBitArray(MaxObjectCount + 1);
-	FNetBitArray ScopeBitArray(MaxObjectCount + 1);
+	FNetBitArray EmptyBitArray(MaxInternalIndex);
+	FNetBitArray ScopeBitArray(MaxInternalIndex);
 
 	// Create an arbitrary number of objects with arbitrary number of frames in between.
-	FNetBitArray OutputArray(MaxObjectCount + 1);
+	FNetBitArray OutputArray(MaxInternalIndex);
 	{
 		FNetBitArrayView OutputBitArrayView = MakeNetBitArrayView(OutputArray);
 
@@ -85,7 +94,8 @@ UE_NET_TEST(ObjectPollFrequencyLimiter, ObjectsToPollAreSpreadOutRegardlessOfUpd
 			{
 				++CreatedObjectCount;
 				FInternalNetRefIndex ObjectIndex = CreatedObjectCount;
-				FrequencyLimiter.SetPollFramePeriod(ObjectIndex, PollFrequency);
+				// + 1 is to get the exact PollFramesBetweenUpdates when calling SetPollFrequency.
+				FrequencyLimiter.SetPollFrequency(ObjectIndex, UpdateFrequency/(PollFramesBetweenUpdates + 1));
 			}
 
 			if (CreatedObjectCount > 0)
@@ -97,14 +107,14 @@ UE_NET_TEST(ObjectPollFrequencyLimiter, ObjectsToPollAreSpreadOutRegardlessOfUpd
 	}
 
 	// Verify the objects get polled evenly.
-	for (uint32 PollIt = 0; PollIt <= PollFrequency; ++PollIt)
+	for (uint32 PollIt = 0; PollIt <= PollFramesBetweenUpdates; ++PollIt)
 	{
 		OutputArray.ClearAllBits();
 		FNetBitArrayView OutputBitArrayView = MakeNetBitArrayView(OutputArray);
 		FrequencyLimiter.Update(MakeNetBitArrayView(ScopeBitArray), MakeNetBitArrayView(EmptyBitArray), OutputBitArrayView);
 
 		uint32 PollCount = OutputArray.CountSetBits();
-		const uint32 ExpectedPollCount = PollFrequencyObjectCount/(PollFrequency + 1U);
+		const uint32 ExpectedPollCount = PollFrequencyObjectCount/(PollFramesBetweenUpdates + 1U);
 		UE_NET_ASSERT_EQ(PollCount, ExpectedPollCount);
 	}
 }
diff --git a/Engine/Source/Runtime/Engine/Private/Net/Iris/ReplicationSystem/EngineReplicationBridge.cpp b/Engine/Source/Runtime/Engine/Private/Net/Iris/ReplicationSystem/EngineReplicationBridge.cpp
--- a/Engine/Source/Runtime/Engine/Private/Net/Iris/ReplicationSystem/EngineReplicationBridge.cpp
+++ b/Engine/Source/Runtime/Engine/Private/Net/Iris/ReplicationSystem/EngineReplicationBridge.cpp
@@ -807,8 +807,6 @@ void UEngineReplicationBridge::SetNetDriver(UNetDriver* const InNetDriver)
 void UEngineReplicationBridge::OnMaxTickRateChanged(UNetDriver* InNetDriver, int32 NewMaxTickRate, int32 OldMaxTickRate)
 {
 	SetMaxTickRate(static_cast<float>(FPlatformMath::Max(InNetDriver->GetNetServerMaxTickRate(), 0)));
-
-	ReinitPollFrequency();
 }
 
 bool UEngineReplicationBridge::RemapPathForPIE(uint32 ConnectionId, FString& Str, bool bReading) const
diff --git a/Engine/Source/Runtime/Engine/Private/Net/Iris/ReplicationSystem/NetActorFactory.cpp b/Engine/Source/Runtime/Engine/Private/Net/Iris/ReplicationSystem/NetActorFactory.cpp
--- a/Engine/Source/Runtime/Engine/Private/Net/Iris/ReplicationSystem/NetActorFactory.cpp
+++ b/Engine/Source/Runtime/Engine/Private/Net/Iris/ReplicationSystem/NetActorFactory.cpp
@@ -683,20 +683,6 @@ TOptional<UNetActorFactory::FWorldInfoData> UNetActorFactory::GetWorldInfo(const
 	return OutData;
 }
 
-float UNetActorFactory::GetPollFrequency(UE::Net::FNetRefHandle RootObjectHandle, UObject* RootObjectInstance)
-{
-	using namespace UE::Net;
-
-	const AActor* Actor = Cast<AActor>(RootObjectInstance);
-	if (!Actor)
-	{
-		// Return the CDO default value
-		return AActor::StaticClass()->GetDefaultObject<AActor>()->GetNetUpdateFrequency();
-	}
-
-	return Actor->GetNetUpdateFrequency();
-}
-
  namespace UE::Net
  {
 
diff --git a/Engine/Source/Runtime/Engine/Private/Net/Iris/ReplicationSystem/NetRootObjectFactory.cpp b/Engine/Source/Runtime/Engine/Private/Net/Iris/ReplicationSystem/NetRootObjectFactory.cpp
--- a/Engine/Source/Runtime/Engine/Private/Net/Iris/ReplicationSystem/NetRootObjectFactory.cpp
+++ b/Engine/Source/Runtime/Engine/Private/Net/Iris/ReplicationSystem/NetRootObjectFactory.cpp
@@ -184,14 +184,6 @@ TOptional<UNetObjectFactory::FWorldInfoData> UNetRootObjectFactory::GetWorldInfo
 	return NullOpt;
 }
 
-float UNetRootObjectFactory::GetPollFrequency(UE::Net::FNetRefHandle RootObjectHandle, UObject* RootObjectInstance)
-{
-	//TODO: Figure what we want to do with the default frequency of these objects.
-	//      Could be time to add a enum (every frame, 30htz, etc.)
-
-	return UE::Net::Private::DefaultRootObjectNetUpdateFrequency;
-}
-
 TUniquePtr<UE::Net::FNetObjectCreationHeader> UNetRootObjectFactory::CreateAndFillHeader(UE::Net::FNetRefHandle Handle)
 {
 	using namespace UE::Net;
diff --git a/Engine/Source/Runtime/Engine/Private/Net/Iris/ReplicationSystem/NetSubObjectFactory.cpp b/Engine/Source/Runtime/Engine/Private/Net/Iris/ReplicationSystem/NetSubObjectFactory.cpp
--- a/Engine/Source/Runtime/Engine/Private/Net/Iris/ReplicationSystem/NetSubObjectFactory.cpp
+++ b/Engine/Source/Runtime/Engine/Private/Net/Iris/ReplicationSystem/NetSubObjectFactory.cpp
@@ -344,12 +344,6 @@ TOptional<UNetObjectFactory::FWorldInfoData> UNetSubObjectFactory::GetWorldInfo(
 	return NullOpt;
 }
 
-float UNetSubObjectFactory::GetPollFrequency(UE::Net::FNetRefHandle RootObjectHandle, UObject* RootObjectInstance)
-{
-	ensureMsgf(false, TEXT("UNetSubObjectFactory::GetPollFrequency called but subobjects should never support this. Instance: %s, NetRefHandle: %s"), *GetNameSafe(RootObjectInstance), *Bridge->PrintObjectFromNetRefHandle(RootObjectHandle));
-	return 100.0f;
-}
-
 //------------------------------------------------------------------------
 
 namespace UE::Net
diff --git a/Engine/Source/Runtime/Engine/Public/Net/Iris/ReplicationSystem/NetActorFactory.h b/Engine/Source/Runtime/Engine/Public/Net/Iris/ReplicationSystem/NetActorFactory.h
--- a/Engine/Source/Runtime/Engine/Public/Net/Iris/ReplicationSystem/NetActorFactory.h
+++ b/Engine/Source/Runtime/Engine/Public/Net/Iris/ReplicationSystem/NetActorFactory.h
@@ -198,8 +198,6 @@ public:
 
 	ENGINE_API virtual TOptional<FWorldInfoData> GetWorldInfo(const FWorldInfoContext& Context) const override;
 
-	ENGINE_API virtual float GetPollFrequency(UE::Net::FNetRefHandle RootObjectHandle, UObject* RootObjectInstance) override;
-
 protected:
 
 	ENGINE_API virtual TUniquePtr<UE::Net::FNetObjectCreationHeader> CreateAndFillHeader(UE::Net::FNetRefHandle Handle) override;
diff --git a/Engine/Source/Runtime/Engine/Public/Net/Iris/ReplicationSystem/NetRootObjectFactory.h b/Engine/Source/Runtime/Engine/Public/Net/Iris/ReplicationSystem/NetRootObjectFactory.h
--- a/Engine/Source/Runtime/Engine/Public/Net/Iris/ReplicationSystem/NetRootObjectFactory.h
+++ b/Engine/Source/Runtime/Engine/Public/Net/Iris/ReplicationSystem/NetRootObjectFactory.h
@@ -72,8 +72,6 @@ public:
 
 	ENGINE_API virtual TOptional<FWorldInfoData> GetWorldInfo(const FWorldInfoContext& Context) const override;
 
-	ENGINE_API virtual float GetPollFrequency(UE::Net::FNetRefHandle RootObjectHandle, UObject* RootObjectInstance) override;
-
 protected:
 
 	// UNetObjectFactory interface
diff --git a/Engine/Source/Runtime/Engine/Public/Net/Iris/ReplicationSystem/NetSubObjectFactory.h b/Engine/Source/Runtime/Engine/Public/Net/Iris/ReplicationSystem/NetSubObjectFactory.h
--- a/Engine/Source/Runtime/Engine/Public/Net/Iris/ReplicationSystem/NetSubObjectFactory.h
+++ b/Engine/Source/Runtime/Engine/Public/Net/Iris/ReplicationSystem/NetSubObjectFactory.h
@@ -97,8 +97,6 @@ public:
 
 	ENGINE_API virtual TOptional<FWorldInfoData> GetWorldInfo(const FWorldInfoContext& Context) const override;
 
-	ENGINE_API virtual float GetPollFrequency(UE::Net::FNetRefHandle RootObjectHandle, UObject* RootObjectInstance) override;
-
 protected:
 
 	ENGINE_API virtual TUniquePtr<UE::Net::FNetObjectCreationHeader> CreateAndFillHeader(UE::Net::FNetRefHandle Handle) override;
diff --git a/Engine/Source/Runtime/Net/Iris/Private/Iris/ReplicationSystem/ObjectPollFrequencyLimiter.cpp b/Engine/Source/Runtime/Net/Iris/Private/Iris/ReplicationSystem/ObjectPollFrequencyLimiter.cpp
--- a/Engine/Source/Runtime/Net/Iris/Private/Iris/ReplicationSystem/ObjectPollFrequencyLimiter.cpp
+++ b/Engine/Source/Runtime/Net/Iris/Private/Iris/ReplicationSystem/ObjectPollFrequencyLimiter.cpp
@@ -2,7 +2,7 @@
 
 #include "ObjectPollFrequencyLimiter.h"
 #include "Iris/Core/IrisProfiler.h"
-#include "HAL/IConsoleManager.h"
+#include "Iris/ReplicationSystem/NetRefHandleManager.h"
 #include "Math/VectorRegister.h"
 #include "Templates/AlignmentTemplates.h"
 
@@ -13,20 +13,18 @@ FObjectPollFrequencyLimiter::FObjectPollFrequencyLimiter()
 {
 }
 
-void FObjectPollFrequencyLimiter::Init(uint32 MaxActiveObjectCount)
+void FObjectPollFrequencyLimiter::Init(FInternalNetRefIndex InMaxInternalIndex, const Private::FNetRefHandleManager* InNetRefHandleManager)
 {
-	// We want to be able to easily update 32 objects at a time.
-	const uint32 StorageObjectCount = Align(MaxActiveObjectCount, 32U);
-
-	// Allocate max amount of memory required.
-	FramesBetweenUpdates.SetNumZeroed(StorageObjectCount);
-	FrameCounters.SetNumZeroed(StorageObjectCount);
+	NetRefHandleManager = InNetRefHandleManager;
+	OnMaxInternalNetRefIndexIncreased(InMaxInternalIndex);
 }
 
 void FObjectPollFrequencyLimiter::Deinit()
 {
 	FramesBetweenUpdates.Empty();
 	FrameCounters.Empty();
+	Frequencies.Empty();
+	NetRefHandleManager = nullptr;
 }
 
 void FObjectPollFrequencyLimiter::OnMaxInternalNetRefIndexIncreased(FInternalNetRefIndex NewMaxInternalIndex)
@@ -35,6 +33,7 @@ void FObjectPollFrequencyLimiter::OnMaxInternalNetRefIndexIncreased(FInternalNet
 
 	FramesBetweenUpdates.SetNumZeroed(StorageObjectCount);
 	FrameCounters.SetNumZeroed(StorageObjectCount);
+	Frequencies.SetNumZeroed(StorageObjectCount);
 }
 
 void FObjectPollFrequencyLimiter::Update(const FNetBitArrayView& RelevantObjects, const FNetBitArrayView& DirtyObjects, FNetBitArrayView& OutObjectsToPoll)
@@ -44,7 +43,7 @@ void FObjectPollFrequencyLimiter::Update(const FNetBitArrayView& RelevantObjects
 	typedef FNetBitArrayView::StorageWordType WordType;
 	constexpr uint32 WordBitCount = FNetBitArrayView::WordBitCount;
 
-	if (!MaxInternalHandle)
+	if (!MaxInternalIndex)
 	{
 		return;
 	}
@@ -68,7 +67,7 @@ void FObjectPollFrequencyLimiter::Update(const FNetBitArrayView& RelevantObjects
 		const VectorRegisterInt Zero = _mm_setzero_si128();
 		const VectorRegisterInt One = _mm_set1_epi8(char(1));
 
-		for (uint32 ObjectIndex = 0, ObjectEndIndex = MaxInternalHandle;
+		for (FInternalNetRefIndex ObjectIndex = 0, ObjectEndIndex = MaxInternalIndex;
 			ObjectIndex <= ObjectEndIndex; 
 			ObjectIndex += WordBitCount, CountersData += WordBitCount, FramesBetweenUpdatesData += WordBitCount, ++RelevantObjectsData, ++DirtyObjectsData, ++ObjectsToPollData)
 		{
@@ -112,7 +111,7 @@ void FObjectPollFrequencyLimiter::Update(const FNetBitArrayView& RelevantObjects
 		const int16_t ShiftAmountsData[8] = { -7, -5, -3, -1, 1, 3, 5, 7 };
 		const int16x8_t ShiftAmounts = vld1q_s16(ShiftAmountsData);
 
-		for (uint32 ObjectIndex = 0, ObjectEndIndex = MaxInternalHandle;
+		for (FInternalNetRefIndex ObjectIndex = 0, ObjectEndIndex = MaxInternalIndex;
 			ObjectIndex <= ObjectEndIndex;
 			ObjectIndex += WordBitCount, CountersData += WordBitCount, FramesBetweenUpdatesData += WordBitCount, ++RelevantObjectsData, ++DirtyObjectsData, ++ObjectsToPollData)
 		{
@@ -158,7 +157,7 @@ void FObjectPollFrequencyLimiter::Update(const FNetBitArrayView& RelevantObjects
 #else
 	// Slower path using scalar integer instructions.
 	{
-		for (uint32 ObjectIndex = 0, ObjectEndIndex = MaxInternalHandle;
+		for (FInternalNetRefIndex ObjectIndex = 0, ObjectEndIndex = MaxInternalIndex;
 			ObjectIndex <= ObjectEndIndex;
 			ObjectIndex += WordBitCount, CountersData += WordBitCount, FramesBetweenUpdatesData += WordBitCount, ++RelevantObjectsData, ++DirtyObjectsData, ++ObjectsToPollData)
 		{
@@ -219,4 +218,29 @@ void FObjectPollFrequencyLimiter::Update(const FNetBitArrayView& RelevantObjects
 #endif
 }
 
+void FObjectPollFrequencyLimiter::ReinitFramesBetweenUpdates()
+{
+	if (!NetRefHandleManager)
+	{
+		return;
+	}
+
+	auto UpdatePollFrequency = [this](uint32 RootObjectIndex)
+	{
+		const float PollFrequency = Frequencies[RootObjectIndex];
+		SetPollFrequency(RootObjectIndex, PollFrequency);
+
+		// Make sure the subobjects are polled the same frame as the root object.
+		for (const FInternalNetRefIndex SubObjectIndex : NetRefHandleManager->GetSubObjects(RootObjectIndex))
+		{
+			SetPollWithObject(RootObjectIndex, SubObjectIndex);
+		}
+	};
+
+	// Update poll frame periods for all root objects
+	const FNetBitArrayView ReplicatedObjects = NetRefHandleManager->GetGlobalScopableInternalIndices();
+	const FNetBitArrayView SubObjects = MakeNetBitArrayView(NetRefHandleManager->GetSubObjectInternalIndices());
+	FNetBitArrayView::ForAllSetBits(ReplicatedObjects, SubObjects, FNetBitArrayView::AndNotOp, UpdatePollFrequency);
+}
+
 }
diff --git a/Engine/Source/Runtime/Net/Iris/Private/Iris/ReplicationSystem/ObjectPollFrequencyLimiter.h b/Engine/Source/Runtime/Net/Iris/Private/Iris/ReplicationSystem/ObjectPollFrequencyLimiter.h
--- a/Engine/Source/Runtime/Net/Iris/Private/Iris/ReplicationSystem/ObjectPollFrequencyLimiter.h
+++ b/Engine/Source/Runtime/Net/Iris/Private/Iris/ReplicationSystem/ObjectPollFrequencyLimiter.h
@@ -2,7 +2,7 @@
 
 #pragma once
 
-#include "CoreTypes.h"
+#include "HAL/Platform.h"
 #include "Containers/Array.h"
 #include "Net/Core/NetBitArray.h"
 
@@ -10,6 +10,10 @@ namespace UE::Net
 {
 	typedef uint32 FInternalNetRefIndex;
 }
+namespace UE::Net::Private
+{
+	class FNetRefHandleManager;
+}
 
 namespace UE::Net::Private
 {
@@ -19,12 +23,15 @@ class FObjectPollFrequencyLimiter
 public:
 	FObjectPollFrequencyLimiter();
 
-	void Init(uint32 MaxActiveObjectCount);
+	void Init(FInternalNetRefIndex MaxInternalIndex, const Private::FNetRefHandleManager* InNetRefHandleManager = nullptr);
 	void Deinit();
 
+	void SetUpdateFrequency(float Frequency);
+	float GetUpdateFrequency() const;
+
 	void OnMaxInternalNetRefIndexIncreased(FInternalNetRefIndex NewMaxInternalIndex);
 
-	void SetPollFramePeriod(FInternalNetRefIndex InternalIndex, uint8 PollFramePeriod);
+	void SetPollFrequency(FInternalNetRefIndex InternalIndex, float PollFrequency);
 
 	void SetPollWithObject(FInternalNetRefIndex ObjectToPollWithInternalIndex, FInternalNetRefIndex InternalIndex);
 
@@ -34,27 +41,56 @@ public:
 	*/
 	void Update(const FNetBitArrayView& RelevantObjects, const FNetBitArrayView& DirtyObjects, FNetBitArrayView& OutObjectsToPoll);
 
-	/** We use a uint8 to track frames, so the limit is 255 frames.*/
-	static constexpr uint32 GetMaxPollingFrames()
-	{		
-		return static_cast<uint32>(std::numeric_limits<uint8>::max());
-	}
+private:
+	void SetPollFramePeriod(FInternalNetRefIndex InternalIndex, uint8 PollFramePeriod);
+	uint8 ConvertFrequencyToFramesBetweenUpdates(float Frequency);
+	void ReinitFramesBetweenUpdates();
 
 private:
 
-	uint32 MaxInternalHandle = 0;
+	const Private::FNetRefHandleManager* NetRefHandleManager = nullptr;
+	FInternalNetRefIndex MaxInternalIndex = 0;
 	uint32 FrameIndex = 0;
+	// Presumed update frequency. Used by ConvertFrequencyToFramesBetweenUpdates.
+	float UpdateFrequency = 30.0f;
 	// We store the number of frames between updates as a byte to be able to process 16 objects per instruction.
 	// This limits polling to at least every 256th frame. At 30Hz this means every 8.5 seconds.
 	uint8 FrameIndexOffsets[256] = {};
 	TArray<uint8> FramesBetweenUpdates;
 	TArray<uint8> FrameCounters;
+	TArray<float> Frequencies;
 };
 
+inline void FObjectPollFrequencyLimiter::SetUpdateFrequency(float InUpdateFrequency)
+{
+	if (InUpdateFrequency <= 0.0f)
+	{
+		// 0 frequency seems incorrect and would cause division by zero.
+		ensureMsgf(false, TEXT("ObjectPollFrequencyLimiter update frequency must be greater than zero"));
+		InUpdateFrequency = 1.0f;
+	}
+
+	if (InUpdateFrequency != UpdateFrequency)
+	{
+		UpdateFrequency = InUpdateFrequency;
+		ReinitFramesBetweenUpdates();
+	}
+}
+
+inline float FObjectPollFrequencyLimiter::GetUpdateFrequency() const
+{
+	return UpdateFrequency;
+}
+
+inline void FObjectPollFrequencyLimiter::SetPollFrequency(FInternalNetRefIndex InternalIndex, float PollFrequency)
+{
+	Frequencies[InternalIndex] = PollFrequency;
+	SetPollFramePeriod(InternalIndex, ConvertFrequencyToFramesBetweenUpdates(PollFrequency));
+}
 
 inline void FObjectPollFrequencyLimiter::SetPollFramePeriod(FInternalNetRefIndex InternalIndex, uint8 PollFramePeriod)
 {
-	MaxInternalHandle = FPlatformMath::Max(MaxInternalHandle, InternalIndex);
+	MaxInternalIndex = FPlatformMath::Max(MaxInternalIndex, InternalIndex);
 
 	FramesBetweenUpdates[InternalIndex] = PollFramePeriod;
 	// Spread the polling of objects with the same frequency so that if you add lots of objects the same frame they won't be polled at the same time. The update loop decrements counters so we need to be careful with how we offset things.
@@ -64,11 +100,31 @@ inline void FObjectPollFrequencyLimiter::SetPollFramePeriod(FInternalNetRefIndex
 
 inline void FObjectPollFrequencyLimiter::SetPollWithObject(FInternalNetRefIndex ObjectToPollWithInternalIndex, FInternalNetRefIndex InternalIndex)
 {
-	MaxInternalHandle = FPlatformMath::Max(MaxInternalHandle, InternalIndex);
+	MaxInternalIndex = FPlatformMath::Max(MaxInternalIndex, InternalIndex);
 
 	// Copy state from object to poll with
 	FramesBetweenUpdates[InternalIndex] = FramesBetweenUpdates[ObjectToPollWithInternalIndex];
 	FrameCounters[InternalIndex] = FrameCounters[ObjectToPollWithInternalIndex];
+
+	// This is really only somewhat useful if this was to be called for a root object rather than a subobject since we will call this method for subobjects in ReinitFramesBetweenUpdates.
+	// For other root objects the SetPollWithObject would have to be called manually when the update frequency changes. This call will only guarantee that we poll as often as the object to poll with, but not necessarily on the same frame.
+	Frequencies[InternalIndex] = Frequencies[ObjectToPollWithInternalIndex];
+}
+
+inline uint8 FObjectPollFrequencyLimiter::ConvertFrequencyToFramesBetweenUpdates(float Frequency)
+{
+	if (Frequency <= 0.0f)
+	{
+		return 0U;
+	}
+
+	uint32 FramesBetweenUpdatesForObject = static_cast<uint32>(UpdateFrequency/FPlatformMath::Max(0.001f, Frequency));
+	if (FramesBetweenUpdatesForObject > 0)
+	{
+		FramesBetweenUpdatesForObject--;
+	}
+
+	return static_cast<uint8>(FMath::Clamp<uint32>(FramesBetweenUpdatesForObject, 0U, 255U));
 }
 
 }
diff --git a/Engine/Source/Runtime/Net/Iris/Private/Iris/ReplicationSystem/ObjectReplicationBridge.cpp b/Engine/Source/Runtime/Net/Iris/Private/Iris/ReplicationSystem/ObjectReplicationBridge.cpp
--- a/Engine/Source/Runtime/Net/Iris/Private/Iris/ReplicationSystem/ObjectReplicationBridge.cpp
+++ b/Engine/Source/Runtime/Net/Iris/Private/Iris/ReplicationSystem/ObjectReplicationBridge.cpp
@@ -231,7 +231,7 @@ void UObjectReplicationBridge::Initialize(UReplicationSystem* InReplicationSyste
 	CreateDestructionProtocol();
 
 	const uint32 CurrentMaxInternalIndex = NetRefHandleManager->GetCurrentMaxInternalNetRefIndex();
-	PollFrequencyLimiter->Init(CurrentMaxInternalIndex);
+	PollFrequencyLimiter->Init(CurrentMaxInternalIndex, NetRefHandleManager);
 	ObjectsWithObjectReferences.Init(CurrentMaxInternalIndex);
 	GarbageCollectionAffectedObjects.Init(CurrentMaxInternalIndex);
 
@@ -826,10 +826,8 @@ UE::Net::FNetRefHandle UObjectReplicationBridge::StartReplicatingRootObject(UObj
 
 	// Set poll frame period
 	float PollFrequency = Params.PollFrequency;
-	FindOrCachePollFrequency(Instance->GetClass(), PollFrequency);
-			
-	uint8 PollFramePeriod = ConvertPollFrequencyIntoFrames(PollFrequency);
-	PollFrequencyLimiter->SetPollFramePeriod(ObjectIndex, PollFramePeriod);
+	FindOrCachePollFrequency(Instance->GetClass(), PollFrequency);			
+	PollFrequencyLimiter->SetPollFrequency(ObjectIndex, PollFrequency);
 
 	if (Params.bNeedsWorldLocationUpdate)
 	{
@@ -2939,6 +2937,12 @@ FString UObjectReplicationBridge::PrintConnectionInfo(uint32 ConnectionId) const
 	return FString::Printf(TEXT("ConnectionId:%u"), ConnectionId);
 }
 
+void UObjectReplicationBridge::SetMaxTickRate(float InMaxTickRate)
+{
+	MaxTickRate = InMaxTickRate;
+	PollFrequencyLimiter->SetUpdateFrequency(InMaxTickRate);
+}
+
 void UObjectReplicationBridge::OptionallySetObjectRequiresFrequentWorldLocationUpdate(FNetRefHandle RefHandle, bool bDesiresFrequentWorldLocationUpdate)
 {
 	using namespace UE::Net;
@@ -3384,75 +3388,8 @@ void UObjectReplicationBridge::InitConditionalPropertyDelegates()
 	});
 }
 
-uint8 UObjectReplicationBridge::ConvertPollFrequencyIntoFrames(float PollFrequency) const
-{
-	if (PollFrequency <= 0.0f)
-	{
-		return 0U;
-	}
-
-	uint32 FramesBetweenUpdatesForObject = static_cast<uint32>(MaxTickRate / FPlatformMath::Max(0.001f, PollFrequency));
-	if (FramesBetweenUpdatesForObject > 0)
-	{
-		FramesBetweenUpdatesForObject--;
-	}
-
-	return static_cast<uint8>(FMath::Clamp<uint32>(FramesBetweenUpdatesForObject, 0U, UE::Net::Private::FObjectPollFrequencyLimiter::GetMaxPollingFrames()));
-}
-
-float UObjectReplicationBridge::GetPollFrequencyOfRootObject(UE::Net::FInternalNetRefIndex RootObjectIndex) const
-{
-	using namespace UE::Net;
-	using namespace UE::Net::Private;
-
-	UObject* RootObjectInstance = NetRefHandleManager->GetReplicatedObjectInstance(RootObjectIndex);
-	FNetRefHandle RootObjectHandle = NetRefHandleManager->GetNetRefHandleFromInternalIndex(RootObjectIndex);
-	
-	const FNetObjectFactoryId FactoryId = NetRefHandleManager->GetReplicatedObjectDataNoCheck(RootObjectIndex).NetFactoryId;
-	if (FactoryId == UE::Net::InvalidNetObjectFactoryId)
-	{
-		// TODO: Investigate how scopeable objects are Invalid
-		return 100.f;
-	}
-	UNetObjectFactory* Factory = GetNetFactory(FactoryId);
-
-	float PollFrequency = Factory->GetPollFrequency(RootObjectHandle, RootObjectInstance);
-
-	if (RootObjectInstance)
-	{
-		GetClassPollFrequency(RootObjectInstance->GetClass(), PollFrequency);
-	}
-
-	return PollFrequency;
-}
-
 void UObjectReplicationBridge::ReinitPollFrequency()
 {
-	using namespace UE::Net;
-	using namespace UE::Net::Private;
-
-	FReplicationSystemInternal* ReplicationSystemInternal = GetReplicationSystem()->GetReplicationSystemInternal();
-	const FNetRefHandleManager& LocalNetRefHandleManager = ReplicationSystemInternal->GetNetRefHandleManager();
-
-	auto UpdatePollFrequency = [this, &LocalNetRefHandleManager](uint32 RootObjectIndex)
-	{
-		const float PollFrequency = GetPollFrequencyOfRootObject(RootObjectIndex);
-		const uint8 PollFramePeriod = ConvertPollFrequencyIntoFrames(PollFrequency);
-
-		PollFrequencyLimiter->SetPollFramePeriod(RootObjectIndex, PollFramePeriod);
-
-		// Make sure the subobjects are polled the same frame as the root object.
-		for (const FInternalNetRefIndex SubObjectIndex : LocalNetRefHandleManager.GetSubObjects(RootObjectIndex))
-		{
-			PollFrequencyLimiter->SetPollWithObject(RootObjectIndex, SubObjectIndex);
-		}
-	};
-
-	// Update frequency on all root objects
-	const FNetBitArrayView ReplicatedObjects = LocalNetRefHandleManager.GetGlobalScopableInternalIndices();
-	const FNetBitArrayView SubObjects = MakeNetBitArrayView(LocalNetRefHandleManager.GetSubObjectInternalIndices());
-
-	FNetBitArrayView::ForAllSetBits(ReplicatedObjects, SubObjects, FNetBitArrayView::AndNotOp, UpdatePollFrequency);
 }
 
 void UObjectReplicationBridge::SetPollFrequency(FNetRefHandle RefHandle, float PollFrequency)
@@ -3469,8 +3406,7 @@ void UObjectReplicationBridge::SetPollFrequency(FNetRefHandle RefHandle, float P
 		return;
 	}
 
-	const uint8 PollFramePeriod = ConvertPollFrequencyIntoFrames(PollFrequency);
-	PollFrequencyLimiter->SetPollFramePeriod(RootObjectIndex, PollFramePeriod);
+	PollFrequencyLimiter->SetPollFrequency(RootObjectIndex, PollFrequency);
 
 	// Make sure the subobjects are polled the same frame as the root object.
 	for (const FInternalNetRefIndex SubObjectIndex : LocalNetRefHandleManager.GetSubObjects(RootObjectIndex))
diff --git a/Engine/Source/Runtime/Net/Iris/Public/Iris/ReplicationSystem/NetObjectFactory.h b/Engine/Source/Runtime/Net/Iris/Public/Iris/ReplicationSystem/NetObjectFactory.h
--- a/Engine/Source/Runtime/Net/Iris/Public/Iris/ReplicationSystem/NetObjectFactory.h
+++ b/Engine/Source/Runtime/Net/Iris/Public/Iris/ReplicationSystem/NetObjectFactory.h
@@ -204,8 +204,9 @@ public:
 	 */
 	virtual TOptional<FWorldInfoData> GetWorldInfo(const FWorldInfoContext& Context) const PURE_VIRTUAL(UNetObjectFactory::GetWorldInfo, return NullOpt;);
 
+	UE_DEPRECATED(5.8, "This method is no longer supported.")
 	/** Return the poll frequency of a root object managed by this factory. */
-	virtual float GetPollFrequency(UE::Net::FNetRefHandle RootObjectHandle, UObject* RootObjectInstance) PURE_VIRTUAL(UNetObjectFactory::GetPollFrequency, return 100.f;);
+	virtual float GetPollFrequency(UE::Net::FNetRefHandle RootObjectHandle, UObject* RootObjectInstance) { return 100.0f; }
 
 public:
 
diff --git a/Engine/Source/Runtime/Net/Iris/Public/Iris/ReplicationSystem/ObjectReplicationBridge.h b/Engine/Source/Runtime/Net/Iris/Public/Iris/ReplicationSystem/ObjectReplicationBridge.h
--- a/Engine/Source/Runtime/Net/Iris/Public/Iris/ReplicationSystem/ObjectReplicationBridge.h
+++ b/Engine/Source/Runtime/Net/Iris/Public/Iris/ReplicationSystem/ObjectReplicationBridge.h
@@ -432,9 +432,11 @@ protected:
 	IRISCORE_API void SetPollWithObject(FNetRefHandle ObjectToPollWith, FNetRefHandle Object);
 
 	/** Transform a poll frequency (in updates per second) to an equivalent number of frames. */
+	UE_DEPRECATED(5.8, "FObjectPollFrequencyLimiter will internally convert poll frequency to its internal representation.")
 	IRISCORE_API uint8 ConvertPollFrequencyIntoFrames(float PollFrequency) const;
 
 	/** Re-initialize the poll frequency of all replicated root objects. */
+	UE_DEPRECATED(5.8, "FObjectPollFrequencyLimiter will re-initialize poll frequencies when its update frequency is changed. This method is a no-op")
 	IRISCORE_API void ReinitPollFrequency();
 
 	/** Re-initialize config-driven parameters found in ObjectReplicationBridgeConfig. */
@@ -478,7 +480,7 @@ protected:
 	float GetMaxTickRate() const { return MaxTickRate; }
 
 	/** Change the max tick rate to match the one from the engine */
-	void SetMaxTickRate(float InMaxTickRate) { MaxTickRate = InMaxTickRate; }
+	IRISCORE_API void SetMaxTickRate(float InMaxTickRate);
 
 	/**
 	 * Parses a list of arguments and returns a list of Connection's that match them
@@ -550,9 +552,6 @@ private:
 	/** Tries to load the classes used in poll period overrides. */
 	void FindClassesInPollPeriodOverrides();
 
-	/** Returns the poll frequency of a specific root object. */
-	float GetPollFrequencyOfRootObject(UE::Net::FInternalNetRefIndex RootObjectIndex) const;
-
 	/** Does all that is necessary to replicate any given UObject. Returns the object's handle if successful  */
 	UE::Net::FNetRefHandle StartReplicatingNetObject(UObject* Instance, UE::Net::EReplicationFragmentTraits Traits, UE::Net::FNetObjectFactoryId NetFactoryId, bool bForceNewHandle);
 
