RemoteObjects: RPC Calls should be wrapped in RemoteExecutor Regions for profiling purposes.

#Synced-CL 52448335

diff --git a/Engine/Source/Runtime/Net/Iris/Private/Iris/ReplicationSystem/NetBlob/NetRPC.cpp b/Engine/Source/Runtime/Net/Iris/Private/Iris/ReplicationSystem/NetBlob/NetRPC.cpp
--- a/Engine/Source/Runtime/Net/Iris/Private/Iris/ReplicationSystem/NetBlob/NetRPC.cpp
+++ b/Engine/Source/Runtime/Net/Iris/Private/Iris/ReplicationSystem/NetBlob/NetRPC.cpp
@@ -35,6 +35,10 @@
 #include "ProfilingDebugging/CsvProfiler.h"
 #include "UObject/CoreNetContext.h"
 
+#if UE_WITH_REMOTE_OBJECT_HANDLE
+#include "UObject/RemoteExecutor.h"
+#endif
+
 DEFINE_LOG_CATEGORY_STATIC(LogIrisRpc, Log, All);
 
 namespace UE::Net::Private
@@ -656,6 +660,18 @@ void FNetRPC::CallFunction(FNetRPCCallContext& CallContext)
 		UE::Net::FClientProfiler::RecordRPC(ActualFunction->GetFName());
 #endif
 
+#if UE_WITH_REMOTE_OBJECT_HANDLE && UE_AUTORTFM
+		UE::RemoteExecutor::OnRegionBeginDelegate.ExecuteIfBound(ActualFunction->GetFName());
+		UE_AUTORTFM_ONABORT()
+		{
+			UE::RemoteExecutor::OnRegionEndDelegate.ExecuteIfBound("Abort");
+		};
+		UE_AUTORTFM_ONCOMMIT()
+		{
+			UE::RemoteExecutor::OnRegionEndDelegate.ExecuteIfBound("Commit");
+		};
+#endif
+
 		UE::Net::FScopedNetContextRPC CallingRPC;
 		UE::Net::Private::FScopedRemoteRPCMode ReceivingRemoteRPC(ActualFunction, UE::Net::Private::ERemoteFunctionMode::Receiving);
 		Object->ProcessEvent(const_cast<UFunction*>(ActualFunction), FunctionParameters);
