Iris - Move RootObjectReplicationParams filling responsibilities from ObjectReplicationBridge to NetObjectFactory

#Synced-CL 52460203

diff --git a/Engine/Plugins/Runtime/MultiServerReplication/Source/MultiServerReplication/Private/MultiServerProxy.cpp b/Engine/Plugins/Runtime/MultiServerReplication/Source/MultiServerReplication/Private/MultiServerProxy.cpp
--- a/Engine/Plugins/Runtime/MultiServerReplication/Source/MultiServerReplication/Private/MultiServerProxy.cpp
+++ b/Engine/Plugins/Runtime/MultiServerReplication/Source/MultiServerReplication/Private/MultiServerProxy.cpp
@@ -795,7 +795,7 @@ void UProxyBackendNetDriver::OnIrisRootObjectPostInit(UE::Net::FNetRefHandle Roo
 	}
 
 	UE_LOG(LogNetProxy, Verbose, TEXT("[%s] Attaching root object %s (%s) to %s"), *GetNameSafe(this), *GetNameSafe(RootObject), *RootHandle.ToString(), *GetNameSafe(ProxyNetDriver));
-	FrontendBridge->StartReplicatingRootObject(RootObject, UObjectReplicationBridge::FRootObjectReplicationParams(), FactoryId);
+	FrontendBridge->StartReplicatingRootObject(RootObject, FactoryId);
 
 	if (AActor* Actor = Cast<AActor>(RootObject))
 	{
diff --git a/Engine/Plugins/Runtime/ReplicationSystemTestPlugin/Source/Private/Tests/ReplicationSystem/ReplicatedTestObject.cpp b/Engine/Plugins/Runtime/ReplicationSystemTestPlugin/Source/Private/Tests/ReplicationSystem/ReplicatedTestObject.cpp
--- a/Engine/Plugins/Runtime/ReplicationSystemTestPlugin/Source/Private/Tests/ReplicationSystem/ReplicatedTestObject.cpp
+++ b/Engine/Plugins/Runtime/ReplicationSystemTestPlugin/Source/Private/Tests/ReplicationSystem/ReplicatedTestObject.cpp
@@ -41,6 +41,8 @@ void UReplicatedTestObjectBridge::Initialize(UReplicationSystem* InReplicationSy
 
 	ReplicatedObjectFactoryId = UE::Net::FNetObjectFactoryRegistry::GetFactoryIdFromName(UReplicatedTestObjectFactory::GetFactoryName());
 	check(ReplicatedObjectFactoryId != UE::Net::InvalidNetObjectFactoryId);
+
+	TestObjectFactory = CastChecked<UReplicatedTestObjectFactory>(GetNetFactory(ReplicatedObjectFactoryId));
 }
 
 const UE::Net::FReplicationInstanceProtocol* UReplicatedTestObjectBridge::GetReplicationInstanceProtocol(FNetRefHandle Handle) const
@@ -58,9 +60,7 @@ const UE::Net::FReplicationInstanceProtocol* UReplicatedTestObjectBridge::GetRep
 UE::Net::FNetRefHandle UReplicatedTestObjectBridge::BeginReplication(UReplicatedTestObject* Instance)
 {
 	// Create NetRefHandle for the registered fragments
-	FRootObjectReplicationParams Params;
-
-	FNetRefHandle Handle = Super::StartReplicatingRootObject(Instance, Params, ReplicatedObjectFactoryId);
+	FNetRefHandle Handle = Super::StartReplicatingRootObject(Instance, ReplicatedObjectFactoryId);
 
 	// This is optional but typically we want to cache at least the NetRefHandle in the game instance to avoid doing map lookups to find it
 	if (Handle.IsValid())
@@ -74,7 +74,8 @@ UE::Net::FNetRefHandle UReplicatedTestObjectBridge::BeginReplication(UReplicated
 UE::Net::FNetRefHandle UReplicatedTestObjectBridge::BeginReplication(UReplicatedTestObject* Instance, const UObjectReplicationBridge::FRootObjectReplicationParams& Params)
 {
 	// Create NetRefHandle for the registered fragments
-	FNetRefHandle Handle = Super::StartReplicatingRootObject(Instance, Params, ReplicatedObjectFactoryId);
+	TestObjectFactory->AddRootObjectReplicationParams(Instance, Params);
+	FNetRefHandle Handle = Super::StartReplicatingRootObject(Instance, ReplicatedObjectFactoryId);
 
 	// This is optional but typically we want to cache at least the NetRefHandle in the game instance to avoid doing map lookups to find it
 	Instance->NetRefHandle = Handle;
diff --git a/Engine/Plugins/Runtime/ReplicationSystemTestPlugin/Source/Private/Tests/ReplicationSystem/ReplicatedTestObject.h b/Engine/Plugins/Runtime/ReplicationSystemTestPlugin/Source/Private/Tests/ReplicationSystem/ReplicatedTestObject.h
--- a/Engine/Plugins/Runtime/ReplicationSystemTestPlugin/Source/Private/Tests/ReplicationSystem/ReplicatedTestObject.h
+++ b/Engine/Plugins/Runtime/ReplicationSystemTestPlugin/Source/Private/Tests/ReplicationSystem/ReplicatedTestObject.h
@@ -16,6 +16,8 @@
 #include "Iris/ReplicationState/ReplicationStateDescriptorMacros.h"
 #include "ReplicatedTestObject.generated.h"
 
+class UReplicatedTestObjectFactory;
+
 namespace UE::Net
 {
 	struct FReplicationInstanceProtocol;
@@ -669,6 +671,7 @@ protected:
 	bool bForceFailCreateRemoteInstance = false;
 
 	UE::Net::FNetObjectFactoryId ReplicatedObjectFactoryId = UE::Net::InvalidNetObjectFactoryId;
+	UReplicatedTestObjectFactory* TestObjectFactory = nullptr;
 };
 
 extern const UE::Net::FRepTag RepTag_FakeGeneratedReplicationState_IntB;
diff --git a/Engine/Plugins/Runtime/ReplicationSystemTestPlugin/Source/Private/Tests/ReplicationSystem/ReplicatedTestObjectFactory.cpp b/Engine/Plugins/Runtime/ReplicationSystemTestPlugin/Source/Private/Tests/ReplicationSystem/ReplicatedTestObjectFactory.cpp
--- a/Engine/Plugins/Runtime/ReplicationSystemTestPlugin/Source/Private/Tests/ReplicationSystem/ReplicatedTestObjectFactory.cpp
+++ b/Engine/Plugins/Runtime/ReplicationSystemTestPlugin/Source/Private/Tests/ReplicationSystem/ReplicatedTestObjectFactory.cpp
@@ -352,4 +352,18 @@ void UReplicatedTestObjectFactory::SubObjectDetachedFromReplication(const FDetac
 	{
 		RootInstance->OnSubObjectDestroyed(Context.DetachedInstance);
 	}
-}
\ No newline at end of file
+}
+
+void UReplicatedTestObjectFactory::FillRootObjectReplicationParams(const UE::Net::FRootObjectReplicationParamsContext& Context, UE::Net::FRootObjectReplicationParams& OutParams)
+{
+	if (const UE::Net::FRootObjectReplicationParams* Params = InstanceReplicationParams.Find(FWeakObjectPtr(Context.Object)))
+	{
+		OutParams = *Params;
+	}
+	// else leave OutParams alone, presumably with default values
+}
+
+void UReplicatedTestObjectFactory::AddRootObjectReplicationParams(const UObject* Object, const UE::Net::FRootObjectReplicationParams& Params)
+{
+	InstanceReplicationParams.Add(FWeakObjectPtr(Object), Params);
+}
diff --git a/Engine/Plugins/Runtime/ReplicationSystemTestPlugin/Source/Private/Tests/ReplicationSystem/ReplicatedTestObjectFactory.h b/Engine/Plugins/Runtime/ReplicationSystemTestPlugin/Source/Private/Tests/ReplicationSystem/ReplicatedTestObjectFactory.h
--- a/Engine/Plugins/Runtime/ReplicationSystemTestPlugin/Source/Private/Tests/ReplicationSystem/ReplicatedTestObjectFactory.h
+++ b/Engine/Plugins/Runtime/ReplicationSystemTestPlugin/Source/Private/Tests/ReplicationSystem/ReplicatedTestObjectFactory.h
@@ -4,6 +4,7 @@
 
 #include "Iris/ReplicationSystem/NetObjectFactory.h"
 #include "Iris/Core/NetObjectReference.h"
+#include "Containers/Map.h"
 
 #include "ReplicatedTestObjectFactory.generated.h"
 
@@ -56,6 +57,8 @@ public:
 	virtual void SubObjectDetachedFromReplication(const FDestroyedContext& Context) override;
 	virtual void SubObjectDetachedFromReplication(const FDetachContext& Context, const FSubObjectDetachContext& SubObjectContext) override;
 
+	virtual void FillRootObjectReplicationParams(const UE::Net::FRootObjectReplicationParamsContext& Context, UE::Net::FRootObjectReplicationParams& OutParams);
+
 	//------------------------------------------------------------------------
 
 	void SetCreatedObjectsOnNode(TArray<TStrongObjectPtr<UObject>>* InCreatedObjectsOnNode)
@@ -70,6 +73,9 @@ public:
 		GetInstanceWorldObjectInfoFunction = LocUpdateFunctor;
 	}
 
+	/** Store instance specific replication params. This factory doesn't support any other interface for retrieving the info through other means. */
+	void AddRootObjectReplicationParams(const UObject*, const UE::Net::FRootObjectReplicationParams& Params);
+
 public:
 
 	TArray<TStrongObjectPtr<UObject>>* CreatedObjectsOnNode = nullptr;
@@ -83,9 +89,10 @@ protected:
 	virtual TUniquePtr<UE::Net::FNetObjectCreationHeader> CreateAndDeserializeHeader(const UE::Net::FCreationHeaderContext& Context) override;
 
 private:
-
-	TFunction<void(UE::Net::FNetRefHandle, const UObject*, FVector&, float&)> GetInstanceWorldObjectInfoFunction;
-};
-
 	
+	TFunction<void(UE::Net::FNetRefHandle, const UObject*, FVector&, float&)> GetInstanceWorldObjectInfoFunction;
 
+private:
+	// Helper for the cases where we actually provide instance specific replication parameters. This is far from ideal but good enough for testing purposes.
+	TSparseMap<FWeakObjectPtr, UE::Net::FRootObjectReplicationParams, FDefaultSparseSetAllocator, TWeakObjectPtrMapKeyFuncs<FWeakObjectPtr, UE::Net::FRootObjectReplicationParams>> InstanceReplicationParams;
+};
diff --git a/Engine/Plugins/Runtime/ReplicationSystemTestPlugin/Source/Private/Tests/ReplicationSystem/ReplicationSystemProxyTestFixture.cpp b/Engine/Plugins/Runtime/ReplicationSystemTestPlugin/Source/Private/Tests/ReplicationSystem/ReplicationSystemProxyTestFixture.cpp
--- a/Engine/Plugins/Runtime/ReplicationSystemTestPlugin/Source/Private/Tests/ReplicationSystem/ReplicationSystemProxyTestFixture.cpp
+++ b/Engine/Plugins/Runtime/ReplicationSystemTestPlugin/Source/Private/Tests/ReplicationSystem/ReplicationSystemProxyTestFixture.cpp
@@ -205,7 +205,7 @@ void FReplicationSystemProxyTestFixture::BindProxyFrontendAndBackend(UObjectRepl
 	GetRootObjectPostInitDelegate(BackendBridge).AddLambda(
 		[FrontendBridge](UE::Net::FNetRefHandle RootHandle, UObject* RootObject, UE::Net::FNetObjectFactoryId FactoryId)
 		{
-			FrontendBridge->StartReplicatingRootObject(RootObject, UObjectReplicationBridge::FRootObjectReplicationParams(), FactoryId);
+			FrontendBridge->StartReplicatingRootObject(RootObject, FactoryId);
 		});
 
 	GetRootObjectDetachedDelegate(BackendBridge).AddLambda(
diff --git a/Engine/Source/Runtime/Engine/Private/Net/Iris/ReplicationSystem/EngineReplicationBridge.cpp b/Engine/Source/Runtime/Engine/Private/Net/Iris/ReplicationSystem/EngineReplicationBridge.cpp
--- a/Engine/Source/Runtime/Engine/Private/Net/Iris/ReplicationSystem/EngineReplicationBridge.cpp
+++ b/Engine/Source/Runtime/Engine/Private/Net/Iris/ReplicationSystem/EngineReplicationBridge.cpp
@@ -395,44 +395,9 @@ UE::Net::FNetRefHandle UEngineReplicationBridge::StartReplicatingActor(AActor* A
 
 	FActorReplicationParams ActorParams;
 	AActor::FFillReplicationParamsContext Context;
+	// Need to call this to get any netfactory override
 	Actor->FillReplicationParams(Context, ActorParams);
 
-	// Create handles for the registered fragments
-	Super::FRootObjectReplicationParams RootObjectParams = 
-	{
-		.bNeedsPreUpdate = true,
-		.bNeedsWorldLocationUpdate = true,
-		.bIsDormant = Actor->NetDormancy > DORM_Awake,
-		.StaticPriority = (Actor->bAlwaysRelevant || Actor->bOnlyRelevantToOwner) ? Actor->NetPriority : 0.0f,
-		.PollFrequency = Actor->GetNetUpdateFrequency(),
-	};
-
-	// Assign a specific filter if requested
-	if (ActorParams.FilterType != FActorReplicationParams::ConfigFilter)
-	{
-		RootObjectParams.bUseExplicitDynamicFilter = true;
-
-		switch(ActorParams.FilterType)
-		{
-			case FActorReplicationParams::AlwaysRelevant:				
-				RootObjectParams.ExplicitDynamicFilterName = NAME_None;
-				break;
-
-			case FActorReplicationParams::DefaultSpatial:
-				RootObjectParams.ExplicitDynamicFilterName = GetDefaultSpatialFilterName();
-				break;
-
-			case FActorReplicationParams::ExplicitFilter:
-				ensureMsgf(!ActorParams.ExplicitDynamicFilterName.IsNone(), TEXT("FActorReplicationParams::ExplicitDynamicFilterName must be defined when FilterType is ExplicitFilter"));
-				RootObjectParams.ExplicitDynamicFilterName = ActorParams.ExplicitDynamicFilterName;
-				break;
-		}
-	}
-
-#if !UE_BUILD_SHIPPING
-	ensureMsgf(!(Actor->bAlwaysRelevant || Actor->bOnlyRelevantToOwner) || RootObjectParams.StaticPriority >= 1.0f, TEXT("Very low NetPriority %.02f for always relevant or owner relevant Actor %s. Set it to 1.0f or higher."), Actor->NetPriority, ToCStr(Actor->GetName()));
-#endif
-
 	// Find the factory ID if this actor does not use the default one.
 	FNetObjectFactoryId FactoryId = ActorFactoryId;
 	if (ActorParams.NetFactoryName != NAME_None)
@@ -443,7 +408,7 @@ UE::Net::FNetRefHandle UEngineReplicationBridge::StartReplicatingActor(AActor* A
 		FactoryId = CustomFactoryId != InvalidNetObjectFactoryId ? CustomFactoryId : ActorFactoryId;
 	}
 
-	FNetRefHandle ActorRefHandle = StartReplicatingRootObject(Actor, RootObjectParams, FactoryId);
+	FNetRefHandle ActorRefHandle = StartReplicatingRootObject(Actor, FactoryId);
 
 	if (!ActorRefHandle.IsValid())
 	{
diff --git a/Engine/Source/Runtime/Engine/Private/Net/Iris/ReplicationSystem/NetActorFactory.cpp b/Engine/Source/Runtime/Engine/Private/Net/Iris/ReplicationSystem/NetActorFactory.cpp
--- a/Engine/Source/Runtime/Engine/Private/Net/Iris/ReplicationSystem/NetActorFactory.cpp
+++ b/Engine/Source/Runtime/Engine/Private/Net/Iris/ReplicationSystem/NetActorFactory.cpp
@@ -100,6 +100,55 @@ void UNetActorFactory::OnInit()
 	SpawnInfoFlags = UE::Net::Private::GetSpawnInfoFlags();
 }
 
+void UNetActorFactory::FillRootObjectReplicationParams(const UE::Net::FRootObjectReplicationParamsContext& Context, UE::Net::FRootObjectReplicationParams& OutParams)
+{
+	using namespace UE::Net;
+
+	AActor* Actor = Cast<AActor>(Context.Object);
+	if (!Actor)
+	{
+		UE_LOG(LogIris, Error, TEXT("UNetActorFactory only supports AActor derived classes. Got %s"), ToCStr(GetNameSafe(Context.Object ? Context.Object->GetClass() : nullptr)));
+		ensureMsgf(false, TEXT("UNetActorFactory only supports AActor derived classes. Got %s"), ToCStr(GetNameSafe(Context.Object ? Context.Object->GetClass() : nullptr)));
+		return;
+	}
+
+	AActor::FFillReplicationParamsContext ActorParamsContext;
+	FActorReplicationParams ActorParams;
+	Actor->FillReplicationParams(ActorParamsContext, ActorParams);
+
+	OutParams.bNeedsPreUpdate = true,
+		OutParams.bNeedsWorldLocationUpdate = true,
+		OutParams.bIsDormant = Actor->NetDormancy > DORM_Awake,
+		OutParams.StaticPriority = (Actor->bAlwaysRelevant || Actor->bOnlyRelevantToOwner) ? Actor->NetPriority : 0.0f,
+		OutParams.PollFrequency = Actor->GetNetUpdateFrequency();
+
+#if !UE_BUILD_SHIPPING
+	ensureMsgf(!(Actor->bAlwaysRelevant || Actor->bOnlyRelevantToOwner) || OutParams.StaticPriority >= 1.0f, TEXT("Very low NetPriority %.02f for always relevant or owner relevant Actor %s. Set it to 1.0f or higher."), Actor->NetPriority, ToCStr(Actor->GetName()));
+#endif
+
+	// Assign a specific filter if requested
+	if (ActorParams.FilterType != FActorReplicationParams::ConfigFilter)
+	{
+		OutParams.bUseExplicitDynamicFilter = true;
+
+		switch (ActorParams.FilterType)
+		{
+		case FActorReplicationParams::AlwaysRelevant:
+			OutParams.ExplicitDynamicFilterName = NAME_None;
+			break;
+
+		case FActorReplicationParams::DefaultSpatial:
+			OutParams.ExplicitDynamicFilterName = Bridge->GetDefaultSpatialFilterName();
+			break;
+
+		case FActorReplicationParams::ExplicitFilter:
+			ensureMsgf(!ActorParams.ExplicitDynamicFilterName.IsNone(), TEXT("FActorReplicationParams::ExplicitDynamicFilterName must be defined when FilterType is ExplicitFilter"));
+			OutParams.ExplicitDynamicFilterName = ActorParams.ExplicitDynamicFilterName;
+			break;
+		}
+	}
+}
+
 TUniquePtr<UE::Net::FNetObjectCreationHeader> UNetActorFactory::CreateAndFillHeader(UE::Net::FNetRefHandle Handle)
 {
 	using namespace UE::Net;
diff --git a/Engine/Source/Runtime/Engine/Private/Net/Iris/ReplicationSystem/NetRootObjectAdapter.cpp b/Engine/Source/Runtime/Engine/Private/Net/Iris/ReplicationSystem/NetRootObjectAdapter.cpp
--- a/Engine/Source/Runtime/Engine/Private/Net/Iris/ReplicationSystem/NetRootObjectAdapter.cpp
+++ b/Engine/Source/Runtime/Engine/Private/Net/Iris/ReplicationSystem/NetRootObjectAdapter.cpp
@@ -174,29 +174,7 @@ void FNetRootObjectAdapter::StartReplication(ULevel* Level)
 
 	FReplicationSystemUtil::ForEachServerBridge(World, [&](UEngineReplicationBridge* Bridge)
 	{
-		UObjectReplicationBridge::FRootObjectReplicationParams Params 
-		{
-			.bUseClassConfigDynamicFilter = false,
-			.PollFrequency = UNetRootObjectFactory::GetDefaultNetUpdateFrequency(),
-		};
-
-		if (bIsNotRouted)
-		{
-			Params.bUseExplicitDynamicFilter = true;
-			Params.ExplicitDynamicFilterName = UE::Net::Private::NotRoutedFilter;
-		}
-		else if (bIsAlwaysRelevant)
-		{
-			Params.bUseExplicitDynamicFilter = true;
-			Params.ExplicitDynamicFilterName = NAME_None;
-		}
-		else
-		{
-			//TODO: Custom filter support
-			//      Spatial filter support
-		}
-
-		const FNetRefHandle NetRefHandle = Bridge->StartReplicatingRootObject(ReplicatedObject, Params, NetFactoryId);
+		const FNetRefHandle NetRefHandle = Bridge->StartReplicatingRootObject(ReplicatedObject, NetFactoryId);
 
 		if (!NetRefHandle.IsValid())
 		{
@@ -451,4 +429,26 @@ void FNetRootObjectAdapter::RemoveRelevantWith(const UObject* OtherReplicatedObj
 	}
 }
 
-} // end namespace UE::Net
\ No newline at end of file
+void FNetRootObjectAdapter::FillRootObjectReplicationParams(const UE::Net::FRootObjectReplicationParamsContext& Context, UE::Net::FRootObjectReplicationParams& OutParams) const
+{
+	OutParams.bUseClassConfigDynamicFilter = false;
+	OutParams.PollFrequency = UNetRootObjectFactory::GetDefaultNetUpdateFrequency();
+
+	if (bIsNotRouted)
+	{
+		OutParams.bUseExplicitDynamicFilter = true;
+		OutParams.ExplicitDynamicFilterName = UE::Net::Private::NotRoutedFilter;
+	}
+	else if (bIsAlwaysRelevant)
+	{
+		OutParams.bUseExplicitDynamicFilter = true;
+		OutParams.ExplicitDynamicFilterName = NAME_None;
+	}
+	else
+	{
+		//TODO: Custom filter support
+		//      Spatial filter support
+	}
+}
+
+} // end namespace UE::Net
diff --git a/Engine/Source/Runtime/Engine/Private/Net/Iris/ReplicationSystem/NetRootObjectFactory.cpp b/Engine/Source/Runtime/Engine/Private/Net/Iris/ReplicationSystem/NetRootObjectFactory.cpp
--- a/Engine/Source/Runtime/Engine/Private/Net/Iris/ReplicationSystem/NetRootObjectFactory.cpp
+++ b/Engine/Source/Runtime/Engine/Private/Net/Iris/ReplicationSystem/NetRootObjectFactory.cpp
@@ -5,22 +5,19 @@
 #include "HAL/IConsoleManager.h"
 
 #include "Iris/Core/IrisLog.h"
-
 #include "Iris/ReplicationSystem/ObjectReplicationBridge.h"
-
-#include "Iris/Serialization/NetBitStreamUtil.h"
 #include "Iris/Serialization/NetBitStreamReader.h"
 #include "Iris/Serialization/NetBitStreamWriter.h"
 #include "Iris/Serialization/NetSerializationContext.h"
 #include "Iris/Serialization/ObjectNetSerializer.h"
 
-#include "Net/DataBunch.h"
+#include "UObject/UObjectGlobals.h"
 
 #include UE_INLINE_GENERATED_CPP_BY_NAME(NetRootObjectFactory)
 
 namespace UE::Net::Private
 {
-	static const FName RootObjectFactoryName(TEXT("RootObjectFactory"));
+	static const FName RootObjectFactoryName("RootObjectFactory");
 
 	float DefaultRootObjectNetUpdateFrequency = 30.f;
 	static FAutoConsoleVariableRef CVarDefaultRootObjectNetUpdateFrequency(
@@ -42,6 +39,8 @@ FName UNetRootObjectFactory::GetFactoryName()
 void UNetRootObjectFactory::OnInit()
 {
 	Super::OnInit();
+
+	FactoryExtensionInterfaceClass = UNetRootObjectFactoryExtension::StaticClass();
 }
 
 UNetObjectFactory::FInstantiateResult UNetRootObjectFactory::InstantiateReplicatedObjectFromHeader(const FInstantiateContext& Context, const UE::Net::FNetObjectCreationHeader* Header)
@@ -177,13 +176,30 @@ void UNetRootObjectFactory::DetachedFromReplication(const FDetachContext& Contex
 
 TOptional<UNetObjectFactory::FWorldInfoData> UNetRootObjectFactory::GetWorldInfo(const FWorldInfoContext& Context) const
 {
-	//TODO: Currently the only way to add support for world infos to a class type is to derive from this factory.
-	//		We should offer a faster customization path without forcing users to declare new concrete factories for every class type.
+	// TODO: Could add support for this via the INetRootObjectFactoryExtension.
 
 	ensureMsgf(false, TEXT("UNetRootObjectFactory::GetWorldInfo does not support object: %s with a world location. Classes that want to support spatial filtering need to implement the function in their own factory."), *GetNameSafe(Context.Instance));
 	return NullOpt;
 }
 
+void UNetRootObjectFactory::FillRootObjectReplicationParams(const UE::Net::FRootObjectReplicationParamsContext& Context, UE::Net::FRootObjectReplicationParams& OutParams)
+{
+	// If the object itself implements INetRootObjectFactoryExtension we get the replication params from there
+	INetRootObjectFactoryExtension* FactoryExtension = static_cast<INetRootObjectFactoryExtension*>(Context.Object->GetInterfaceAddress(FactoryExtensionInterfaceClass));
+
+	if (FactoryExtension)
+	{
+		FactoryExtension->FillRootObjectReplicationParams(Context, OutParams);
+		return;
+	}
+
+	// It should suffice to ensure once. A single replicated instance of this class would ensure and assuming any testing whatsoever was performed the ensure should have triggered then rather than being hit in a live environment for the first time.
+	ensureMsgf(false, TEXT("Class %s doesn't have an INetRootObjectFactoryExtension associated with it."), ToCStr(Context.Object->GetClass()->GetName()));
+
+	// Use our factory defaults when they differ from FRootObjectReplicationParams defaults.
+	OutParams.PollFrequency = UE::Net::Private::DefaultRootObjectNetUpdateFrequency;
+}
+
 TUniquePtr<UE::Net::FNetObjectCreationHeader> UNetRootObjectFactory::CreateAndFillHeader(UE::Net::FNetRefHandle Handle)
 {
 	using namespace UE::Net;
diff --git a/Engine/Source/Runtime/Engine/Private/Net/Iris/ReplicationSystem/NetSubObjectFactory.cpp b/Engine/Source/Runtime/Engine/Private/Net/Iris/ReplicationSystem/NetSubObjectFactory.cpp
--- a/Engine/Source/Runtime/Engine/Private/Net/Iris/ReplicationSystem/NetSubObjectFactory.cpp
+++ b/Engine/Source/Runtime/Engine/Private/Net/Iris/ReplicationSystem/NetSubObjectFactory.cpp
@@ -344,6 +344,11 @@ TOptional<UNetObjectFactory::FWorldInfoData> UNetSubObjectFactory::GetWorldInfo(
 	return NullOpt;
 }
 
+void UNetSubObjectFactory::FillRootObjectReplicationParams(const UE::Net::FRootObjectReplicationParamsContext& Context, UE::Net::FRootObjectReplicationParams&)
+{
+	ensureMsgf(false, TEXT("UNetSubObjectFactory::FillRootObjectReplicationParams called but subobjects should never support this. Instance: %s"), *GetNameSafe(Context.Object));
+}
+
 //------------------------------------------------------------------------
 
 namespace UE::Net
diff --git a/Engine/Source/Runtime/Engine/Public/Net/Iris/ReplicationSystem/NetActorFactory.h b/Engine/Source/Runtime/Engine/Public/Net/Iris/ReplicationSystem/NetActorFactory.h
--- a/Engine/Source/Runtime/Engine/Public/Net/Iris/ReplicationSystem/NetActorFactory.h
+++ b/Engine/Source/Runtime/Engine/Public/Net/Iris/ReplicationSystem/NetActorFactory.h
@@ -200,6 +200,8 @@ public:
 
 protected:
 
+	ENGINE_API virtual void FillRootObjectReplicationParams(const UE::Net::FRootObjectReplicationParamsContext& Context, UE::Net::FRootObjectReplicationParams& OutParams) override;
+
 	ENGINE_API virtual TUniquePtr<UE::Net::FNetObjectCreationHeader> CreateAndFillHeader(UE::Net::FNetRefHandle Handle) override;
 	ENGINE_API virtual TUniquePtr<UE::Net::FNetObjectCreationHeader> CreateAndDeserializeHeader(const UE::Net::FCreationHeaderContext& Context) override;
 
diff --git a/Engine/Source/Runtime/Engine/Public/Net/Iris/ReplicationSystem/NetRootObjectAdapter.h b/Engine/Source/Runtime/Engine/Public/Net/Iris/ReplicationSystem/NetRootObjectAdapter.h
--- a/Engine/Source/Runtime/Engine/Public/Net/Iris/ReplicationSystem/NetRootObjectAdapter.h
+++ b/Engine/Source/Runtime/Engine/Public/Net/Iris/ReplicationSystem/NetRootObjectAdapter.h
@@ -11,6 +11,11 @@
 
 class UWorld;
 class ULevel;
+namespace UE::Net
+{
+	struct FRootObjectReplicationParamsContext;
+	struct FRootObjectReplicationParams;
+}
 
 namespace UE::Net
 {
@@ -40,6 +45,7 @@ struct FRootObjectSettings
  * 2) Configure the replication settings via Configure() prior to starting replication
  * 3) Provide the ULevel the UObject is a part of
  * 4) Make the object replicated via StartReplication or RelevantWith
+ * 5) Implement a INetRootObjectFactoryExtension either in the class itself or as a standalone class able to retrieve this adapter and call FillRootObjectReplicationParams on it. 
  * 
  * This class is meant to provide an easy to use solution for quickly transforming any non-actor into an autonomous replicated object.
  * But be aware it is not made to be the optimal solution if you intend to create alot of instances of the class using it.
@@ -68,13 +74,13 @@ public:
 	/** 
 	 * Every object needs to be assigned to a ULevel to know which UWorld->ReplicationSystem it will be added to. 
 	 * The level will also be used to set level filtering logic to replicate the object only to clients that loaded the level locally.
-	 * This must be called early if your object will be replicated conditionnally via RelevantWith
+	 * This must be called early if your object will be replicated conditionally via RelevantWith
 	 * Note: if you don't have a specific level you can use the World->PersistentLevel.
 	 */
 	ENGINE_API void SetAttachedLevel(ULevel* Level);
 
 	/** 
-	 * Start replicating the object.
+	 * Start replicating the object. This will cause the UNetRootObjectFactory to look for an INetRootObjectFactoryExtension for the class to retrieve the replication details.
 	 * @param Level: Level the object is a part of. Can be set to null if SetAttachedLevel was previously set.
 	 */
 	ENGINE_API void StartReplication(ULevel* Level);
@@ -101,6 +107,12 @@ public:
 	/** Set to use a netobject factory different from the default. Must be set before replication starts. */
 	ENGINE_API void SetNetFactoryName(FName InFactoryName);
 
+	/**
+	 * Utility function that fills in the FRootObjectReplicationParams. It can be called from an INetRootObjectFactoryExtension interface for example when starting replication or a replay recording. 
+	 * @see INetRootObjectFactoryExtension
+	 */
+	ENGINE_API void FillRootObjectReplicationParams(const UE::Net::FRootObjectReplicationParamsContext& Context, UE::Net::FRootObjectReplicationParams& OutParams) const;
+
 	// Status / Config getters
 
 	bool IsInitialized() const
diff --git a/Engine/Source/Runtime/Engine/Public/Net/Iris/ReplicationSystem/NetRootObjectFactory.h b/Engine/Source/Runtime/Engine/Public/Net/Iris/ReplicationSystem/NetRootObjectFactory.h
--- a/Engine/Source/Runtime/Engine/Public/Net/Iris/ReplicationSystem/NetRootObjectFactory.h
+++ b/Engine/Source/Runtime/Engine/Public/Net/Iris/ReplicationSystem/NetRootObjectFactory.h
@@ -3,12 +3,11 @@
 #pragma once
 
 #include "Iris/ReplicationSystem/NetObjectFactory.h"
-
 #include "Iris/Core/NetObjectReference.h"
+#include "UObject/Interface.h"
 
 #include "NetRootObjectFactory.generated.h"
 
-
 namespace UE::Net
 {
 
@@ -43,9 +42,28 @@ public:
 } // end namespace UE::Net
 
 
+UINTERFACE(MinimalApi, meta=(CannotImplementInterfaceInBlueprint))
+class UNetRootObjectFactoryExtension : public UInterface
+{
+	GENERATED_BODY()
+};
+
+/**
+ * Interface allowing for instance or class specific overrides.
+ * If implemented as a standalone interface and configured via FNetRootObjectFactoryExtensionConfig as part of the UNetRootObjectFactoryConfig one need to inherit both from UObject and INetRootObjectFactoryExtension.
+ * The UNetRootObjectFactory will instantiate a new object from the configured class.
+ */
+class INetRootObjectFactoryExtension
+{
+	GENERATED_BODY()
+
+public:
+	virtual void FillRootObjectReplicationParams(const UE::Net::FRootObjectReplicationParamsContext& Context, UE::Net::FRootObjectReplicationParams& OutParams) const = 0;
+};
+
 /**
  * Factory that can be used for any UObject based class that wants to be replicated autonomously.
- * Can be overriden if the class implementation requires more advanced control over discovery, instantiation, location, etc.
+ * Can be overridden if the class implementation requires more advanced control over discovery, instantiation, location, etc.
  */
 UCLASS(MinimalAPI)
 class UNetRootObjectFactory : public UNetObjectFactory
@@ -57,7 +75,7 @@ public:
 	ENGINE_API static FName GetFactoryName();
 	ENGINE_API static float GetDefaultNetUpdateFrequency();
 
-public:
+protected:
 
 	// UNetObjectFactory interface
 	ENGINE_API virtual void OnInit() override;
@@ -72,14 +90,13 @@ public:
 
 	ENGINE_API virtual TOptional<FWorldInfoData> GetWorldInfo(const FWorldInfoContext& Context) const override;
 
-protected:
-
-	// UNetObjectFactory interface
 	ENGINE_API virtual TUniquePtr<UE::Net::FNetObjectCreationHeader> CreateAndFillHeader(UE::Net::FNetRefHandle Handle) override;
 
 	ENGINE_API virtual TUniquePtr<UE::Net::FNetObjectCreationHeader> CreateAndDeserializeHeader(const UE::Net::FCreationHeaderContext& Context) override;
 
-	ENGINE_API virtual bool SerializeHeader(const UE::Net::FCreationHeaderContext& Context, const UE::Net::FNetObjectCreationHeader* Header)  override;
+	ENGINE_API virtual bool SerializeHeader(const UE::Net::FCreationHeaderContext& Context, const UE::Net::FNetObjectCreationHeader* Header) override;
+
+	ENGINE_API virtual void FillRootObjectReplicationParams(const UE::Net::FRootObjectReplicationParamsContext& Context, UE::Net::FRootObjectReplicationParams& OutParams) override;
 
 protected:
 
@@ -88,4 +105,7 @@ protected:
 	/** Fill the basic header info for a given object */
 	ENGINE_API bool FillHeader(UE::Net::FNetRootObjectCreationHeader* RootObjectHeader, UE::Net::FNetRefHandle Handle, const UObject* RootObjectPtr);
 
+protected:
+	// Cached class of INetRootObjectFactoryExtension
+	UClass* FactoryExtensionInterfaceClass = nullptr;
 };
diff --git a/Engine/Source/Runtime/Engine/Public/Net/Iris/ReplicationSystem/NetSubObjectFactory.h b/Engine/Source/Runtime/Engine/Public/Net/Iris/ReplicationSystem/NetSubObjectFactory.h
--- a/Engine/Source/Runtime/Engine/Public/Net/Iris/ReplicationSystem/NetSubObjectFactory.h
+++ b/Engine/Source/Runtime/Engine/Public/Net/Iris/ReplicationSystem/NetSubObjectFactory.h
@@ -99,10 +99,13 @@ public:
 
 protected:
 
+	// FillRootObjectReplicationParams should not be called for subobjects
+	ENGINE_API virtual void FillRootObjectReplicationParams(const UE::Net::FRootObjectReplicationParamsContext& Context, UE::Net::FRootObjectReplicationParams& OutParams) override final;
+
 	ENGINE_API virtual TUniquePtr<UE::Net::FNetObjectCreationHeader> CreateAndFillHeader(UE::Net::FNetRefHandle Handle) override;
 	ENGINE_API virtual TUniquePtr<UE::Net::FNetObjectCreationHeader> CreateAndDeserializeHeader(const UE::Net::FCreationHeaderContext& Context) override;
 
-	ENGINE_API virtual bool SerializeHeader(const UE::Net::FCreationHeaderContext& Context, const UE::Net::FNetObjectCreationHeader* Header)  override;
+	ENGINE_API virtual bool SerializeHeader(const UE::Net::FCreationHeaderContext& Context, const UE::Net::FNetObjectCreationHeader* Header) override;
 
 	/** Write the needed header information for a dynamic subobject */
 	ENGINE_API bool FillDynamicHeader(UE::Net::FNetDynamicSubObjectCreationHeader* DynamicHeader, const UObject* SubObject, UE::Net::FNetRefHandle Handle);
diff --git a/Engine/Source/Runtime/Net/Iris/Private/Iris/ReplicationSystem/ObjectReplicationBridge.cpp b/Engine/Source/Runtime/Net/Iris/Private/Iris/ReplicationSystem/ObjectReplicationBridge.cpp
--- a/Engine/Source/Runtime/Net/Iris/Private/Iris/ReplicationSystem/ObjectReplicationBridge.cpp
+++ b/Engine/Source/Runtime/Net/Iris/Private/Iris/ReplicationSystem/ObjectReplicationBridge.cpp
@@ -653,7 +653,7 @@ UE::Net::FNetRefHandle UObjectReplicationBridge::StartReplicatingNetObject(UObje
 
 		if (!AllocatedRefHandle.IsValid())
 		{
-			ensureMsgf(AllocatedRefHandle.IsValid(), TEXT("Reassignement for %s failed to return a valid handle"), *GetNameSafe(Instance));
+			ensureMsgf(AllocatedRefHandle.IsValid(), TEXT("Reassignment for %s failed to return a valid handle"), *GetNameSafe(Instance));
 			return FNetRefHandle::GetInvalid();
 		}
 	}
@@ -791,7 +791,7 @@ UE::Net::FNetRefHandle UObjectReplicationBridge::StartReplicatingNetObject(UObje
 	return RefHandle;
 }
 
-UE::Net::FNetRefHandle UObjectReplicationBridge::StartReplicatingRootObject(UObject* Instance, const FRootObjectReplicationParams& Params, UE::Net::FNetObjectFactoryId NetFactoryId)
+UE::Net::FNetRefHandle UObjectReplicationBridge::InternalStartReplicatingRootObject(UObject* Instance, const FRootObjectReplicationParams& Params, UE::Net::FNetObjectFactoryId NetFactoryId)
 {
 	LLM_SCOPE_BYTAG(IrisState);
 
@@ -869,6 +869,34 @@ UE::Net::FNetRefHandle UObjectReplicationBridge::StartReplicatingRootObject(UObj
 	return RefHandle;
 }
 
+UE::Net::FNetRefHandle UObjectReplicationBridge::StartReplicatingRootObject(UObject* Instance, const FRootObjectReplicationParams& Params, UE::Net::FNetObjectFactoryId NetFactoryId)
+{
+	return InternalStartReplicatingRootObject(Instance, Params, NetFactoryId);
+}
+
+UE::Net::FNetRefHandle UObjectReplicationBridge::StartReplicatingRootObject(UObject* Instance, UE::Net::FNetObjectFactoryId NetFactoryId)
+{
+	using namespace UE::Net;
+
+	LLM_SCOPE_BYTAG(IrisState);
+	
+	if (UNetObjectFactory* NetObjectFactory = GetNetFactory(NetFactoryId))
+	{
+		FRootObjectReplicationParamsContext ParamsContext;
+		FRootObjectReplicationParams Params;
+		ParamsContext.Object = Instance;
+		ParamsContext.Reason = ERootObjectReplicationParamsReason::StartReplication;
+		NetObjectFactory->FillRootObjectReplicationParams(ParamsContext, Params);
+
+		return InternalStartReplicatingRootObject(Instance, Params, NetFactoryId);
+	}
+	else
+	{
+		UE_LOG_BRIDGEID(LogIris, Warning, TEXT("%hs: Unknown NetObjectFactory (Id: %u)"), __FUNCTION__, NetFactoryId); 
+		return FNetRefHandle::GetInvalid();
+	}
+}
+
 void UObjectReplicationBridge::AssignDynamicFilter(UObject* Instance, const FRootObjectReplicationParams& Params, FNetRefHandle RefHandle)
 {
 	using namespace UE::Net;
@@ -902,7 +930,6 @@ void UObjectReplicationBridge::AssignDynamicFilter(UObject* Instance, const FRoo
 		UE_LOG_OBJECTREPLICATIONBRIDGE(Verbose, TEXT("StartReplicatingRootObject Filter: %s will be used for %s. (FilterProfile: %s)"), *(ReplicationSystem->GetFilterName(FilterHandle).ToString()), *NetRefHandleManager->PrintObjectFromNetRefHandle(RefHandle), *FilterConfigProfile.ToString());
 		ReplicationSystem->SetFilter(RefHandle, FilterHandle, FilterConfigProfile);
 	}
-	
 }
 
 UE::Net::FNetRefHandle UObjectReplicationBridge::StartReplicatingSubObject(UObject* Instance, const FSubObjectReplicationParams& Params, UE::Net::FNetObjectFactoryId NetFactoryId)
diff --git a/Engine/Source/Runtime/Net/Iris/Public/Iris/ReplicationSystem/NetObjectFactory.h b/Engine/Source/Runtime/Net/Iris/Public/Iris/ReplicationSystem/NetObjectFactory.h
--- a/Engine/Source/Runtime/Net/Iris/Public/Iris/ReplicationSystem/NetObjectFactory.h
+++ b/Engine/Source/Runtime/Net/Iris/Public/Iris/ReplicationSystem/NetObjectFactory.h
@@ -86,7 +86,62 @@ public:
 private:
 
 	FReplicationProtocolIdentifier ProtocolIdentifier = 0;
-	FNetObjectFactoryId  FactoryId = InvalidNetObjectFactoryId;
+	FNetObjectFactoryId FactoryId = InvalidNetObjectFactoryId;
+};
+
+enum class ERootObjectReplicationParamsReason : uint32
+{
+	StartReplication,
+	StartReplayRecording,
+};
+
+struct FRootObjectReplicationParamsContext
+{
+	/** Object params are being requested for. */
+	UObject* Object = nullptr;
+
+	/** The reason for FillRootObjectReplicationParams being called. */
+	ERootObjectReplicationParamsReason Reason = ERootObjectReplicationParamsReason::StartReplication;
+};
+
+struct FRootObjectReplicationParams
+{
+	/** *Actor only* When true the actor will receive a PreReplication callback just before it gets polled. */
+	bool bNeedsPreUpdate = false;
+
+	/** When true the object has a dynamic world location and we should ask the factory to update its current location every time it is polled. */
+	bool bNeedsWorldLocationUpdate = false;
+
+	/** Whether the object is dormant or not */
+	bool bIsDormant = false;
+
+	/** Ask the class config for a dynamic filter assigned to this class or one of it's parent class. Default is true. */
+	bool bUseClassConfigDynamicFilter = true;
+
+	/** When enabled we ignore the class config for this object and instead use the one specified by ExplicitDynamicFilter */
+	bool bUseExplicitDynamicFilter = false;
+
+	/** Set to true to prevent the replication system from reusing a NetRefHandle previously assigned to the object. Useful if the object had stopped replicating and that old representation is considered to be destroyed. */
+	bool bForceNewHandle = false;
+
+	/**
+	* The name of the dynamic filter to use for this object (instead of asking the class config).
+	* Can be none so that no dynamic filter is assigned to the object.
+	* Only used when bUseExplicitDynamicFilter is true.
+	*/
+	FName ExplicitDynamicFilterName;
+
+	/**
+	* If StaticPriority is > 0 the ReplicationSystem will use that as priority when scheduling objects.
+	* If it's <= 0.0f one will look for a world location support and then use the default spatial prioritizer.
+	*/
+	float StaticPriority = 0.0f;
+
+	/**
+	* How often per second the object should be polled for dirtiness, including calling the InstancePreUpdate function.
+	* When set to zero it will be polled every frame.
+	*/
+	float PollFrequency = 0.0f;
 };
 
 } // end namespace UE::Net
@@ -208,6 +263,8 @@ public:
 	/** Return the poll frequency of a root object managed by this factory. */
 	virtual float GetPollFrequency(UE::Net::FNetRefHandle RootObjectHandle, UObject* RootObjectInstance) { return 100.0f; }
 
+	virtual void FillRootObjectReplicationParams(const UE::Net::FRootObjectReplicationParamsContext& Context, UE::Net::FRootObjectReplicationParams& OutParams) PURE_VIRTUAL(UNetObjectFactory::FillRootObjectReplicationParams);
+
 public:
 
 	/** Result of the instantiate request */
diff --git a/Engine/Source/Runtime/Net/Iris/Public/Iris/ReplicationSystem/ObjectReplicationBridge.h b/Engine/Source/Runtime/Net/Iris/Public/Iris/ReplicationSystem/ObjectReplicationBridge.h
--- a/Engine/Source/Runtime/Net/Iris/Public/Iris/ReplicationSystem/ObjectReplicationBridge.h
+++ b/Engine/Source/Runtime/Net/Iris/Public/Iris/ReplicationSystem/ObjectReplicationBridge.h
@@ -121,7 +121,7 @@ public:
 
 	using EGetRefHandleFlags = UE::Net::EGetRefHandleFlags;
 
-	struct FRootObjectReplicationParams;
+	using FRootObjectReplicationParams = UE::Net::FRootObjectReplicationParams;
 	struct FSubObjectReplicationParams;
 
 	IRISCORE_API UObjectReplicationBridge();
@@ -155,8 +155,18 @@ public:
 	 * @param NetFactoryId NetFactory to use for this object.
 	 * @return The NetRefHandle associated to this object if the operation succeeded.
 	 */
+	UE_DEPRECATED(5.8, "Use version not taking a FRootObjectReplicationParams parameter. The responsibility of filling in these params has moved to UNetObjectFactory.")
 	IRISCORE_API FNetRefHandle StartReplicatingRootObject(UObject* Instance, const FRootObjectReplicationParams& Params, UE::Net::FNetObjectFactoryId NetFactoryId);
 
+	/** 
+	 * Start replicating a RootObject and return a valid NetRefHandle if successful. 
+	 * 
+	 * @param Instance The instance that needs to be replicated
+	 * @param NetFactoryId NetFactory to use for this object.
+	 * @return The NetRefHandle associated to this object if the operation succeeded.
+	 */
+	IRISCORE_API FNetRefHandle StartReplicatingRootObject(UObject* Instance, UE::Net::FNetObjectFactoryId NetFactoryId);
+
 	/**
 	 * Start replicating a SubObject and return a valid NetRefHandle if successful.
 	 * 
@@ -291,45 +301,6 @@ public:
 
 public:
 
-	struct FRootObjectReplicationParams
-	{
-		/** *Actor only* When true the actor will receive a PreReplication callback just before it gets polled. */
-		bool bNeedsPreUpdate = false;
-
-		/** When true the object has a dynamic world location and we should ask the bridge to update its current location every time it is polled. */
-		bool bNeedsWorldLocationUpdate = false;
-
-		/** Whether the object is dormant or not */
-		bool bIsDormant = false;
-
-		/** Ask the class config for a dynamic filter assigned to this class or one of it's parent class. Default is true. */
-		bool bUseClassConfigDynamicFilter = true;
-
-		/** When enabled we ignore the class config for this object and instead use the one specified by ExplicitDynamicFilter */
-		bool bUseExplicitDynamicFilter = false;
-
-		/** Set to true to prevent the replication system from reusing a NetRefHandle previously assigned to the object. Useful if the object had stopped replicating and that old representation is considered to be destroyed. */
-		bool bForceNewHandle = false;
-
-		/** 
-		* The name of the dynamic filter to use for this object (instead of asking the class config).
-		* Can be none so that no dynamic filter is assigned to the object.
-		* Only used when bUseExplicitDynamicFilter is true.
-		*/
-		FName ExplicitDynamicFilterName;
-
-		/**
-		* If StaticPriority is > 0 the ReplicationSystem will use that as priority when scheduling objects. 
-		* If it's <= 0.0f one will look for a world location support and then use the default spatial prioritizer.
-		*/
-		float StaticPriority = 0.0f;
-
-		/**
-		* How often per second the object should be polled for dirtiness, including calling the InstancePreUpdate function. 
-		* When set to zero it will be polled every frame.
-		*/
-		float PollFrequency = 0.0f;
-	};
 
 	struct FSubObjectReplicationParams
 	{
@@ -552,6 +523,8 @@ private:
 	/** Tries to load the classes used in poll period overrides. */
 	void FindClassesInPollPeriodOverrides();
 
+	FNetRefHandle InternalStartReplicatingRootObject(UObject* Instance, const FRootObjectReplicationParams& Params, UE::Net::FNetObjectFactoryId NetFactoryId);
+
 	/** Does all that is necessary to replicate any given UObject. Returns the object's handle if successful  */
 	UE::Net::FNetRefHandle StartReplicatingNetObject(UObject* Instance, UE::Net::EReplicationFragmentTraits Traits, UE::Net::FNetObjectFactoryId NetFactoryId, bool bForceNewHandle);
 
