Iris - Mute experimental warnings in experimental code

#Synced-CL 52460203

diff --git a/Engine/Source/Runtime/Net/Iris/Public/Iris/Config/ScopedIrisDynamicConfig.h b/Engine/Source/Runtime/Net/Iris/Public/Iris/Config/ScopedIrisDynamicConfig.h
--- a/Engine/Source/Runtime/Net/Iris/Public/Iris/Config/ScopedIrisDynamicConfig.h
+++ b/Engine/Source/Runtime/Net/Iris/Public/Iris/Config/ScopedIrisDynamicConfig.h
@@ -23,32 +23,42 @@ public:
 	void RegisterDynamicConfig(FName UniqueConfigName, const FString& FilePath)
 	{
 		RegisteredConfigs.Add(UniqueConfigName);
+PRAGMA_DISABLE_EXPERIMENTAL_WARNINGS
 		FIrisDynamicConfig::RegisterDynamicConfig(UniqueConfigName, FilePath);
+PRAGMA_ENABLE_EXPERIMENTAL_WARNINGS
 	}
 
 	void RegisterDynamicConfigBuffer(FName UniqueConfigName, const FString& VirtualFilePath, const FString& ConfigContents)
 	{
 		RegisteredConfigs.Add(UniqueConfigName);
+PRAGMA_DISABLE_EXPERIMENTAL_WARNINGS
 		FIrisDynamicConfig::RegisterDynamicConfigBuffer(UniqueConfigName, VirtualFilePath, ConfigContents);
+PRAGMA_ENABLE_EXPERIMENTAL_WARNINGS
 	}
 
 	void RegisterDynamicConfigBuffer(FName UniqueConfigName, const FString& VirtualFilePath, FUtf8StringView ConfigContents)
 	{
 		RegisteredConfigs.Add(UniqueConfigName);
+PRAGMA_DISABLE_EXPERIMENTAL_WARNINGS
 		FIrisDynamicConfig::RegisterDynamicConfigBuffer(UniqueConfigName, VirtualFilePath, ConfigContents);
+PRAGMA_ENABLE_EXPERIMENTAL_WARNINGS
 	}
 
 	void UnregisterDynamicConfig(FName UniqueConfigName)
 	{
 		RegisteredConfigs.Remove(UniqueConfigName);
+PRAGMA_DISABLE_EXPERIMENTAL_WARNINGS
 		FIrisDynamicConfig::UnregisterDynamicConfig(UniqueConfigName);
+PRAGMA_ENABLE_EXPERIMENTAL_WARNINGS
 	}
 
 	void UnregisterAllDynamicConfigs()
 	{
 		for (FName ConfigName : RegisteredConfigs)
 		{
+PRAGMA_DISABLE_EXPERIMENTAL_WARNINGS
 			FIrisDynamicConfig::UnregisterDynamicConfig(ConfigName);
+PRAGMA_ENABLE_EXPERIMENTAL_WARNINGS
 		}
 		RegisteredConfigs.Empty();
 	}
