[iris] FNetBitArray perf: use CTZ in more places

#rnx
#Synced-CL 52447090

diff --git a/Engine/Source/Runtime/Net/Core/Public/Net/Core/NetBitArray.h b/Engine/Source/Runtime/Net/Core/Public/Net/Core/NetBitArray.h
--- a/Engine/Source/Runtime/Net/Core/Public/Net/Core/NetBitArray.h
+++ b/Engine/Source/Runtime/Net/Core/Public/Net/Core/NetBitArray.h
@@ -891,16 +891,12 @@ private:
 		{
 			StorageWordType CurrentWord = Storage[WordIt] & ((WordIt == LastWordIt) ? LastWordMask : ~0U);
 
-			// Test All bits in the CurrentWord and invoke functor if they are set
-			uint32 LocalBitIndex = CurrentBitIndex;
+			// Invoke functor for all set bits in the CurrentWord
 			while (CurrentWord)
 			{
-				if (CurrentWord & 0x1)
-				{
-					Functor(LocalBitIndex);
-				}
-				CurrentWord >>= 1;
-				++LocalBitIndex;
+				const uint32 BitIndex = CurrentBitIndex + FPlatformMath::CountTrailingZeros(CurrentWord);
+				Functor(BitIndex);
+				CurrentWord &= CurrentWord - 1;
 			}
 		}
 	}
@@ -927,16 +923,12 @@ private:
 			StorageWordType WordMask = (~StorageWordType(0) & ((WordIt == StartWord) ? StartWordMask : ~0U)) & ((WordIt == EndWord) ? EndWordMask : ~0U);
 			StorageWordType CurrentWord = Storage[WordIt] & WordMask;
 
-			// Test All bits in the CurrentWord and invoke functor if they are set
-			uint32 LocalBitIndex = CurrentBitIndex;
+			// Invoke functor for all set bits in the CurrentWord
 			while (CurrentWord)
 			{
-				if (CurrentWord & 0x1)
-				{
-					Functor(LocalBitIndex);
-				}
-				CurrentWord >>= 1;
-				++LocalBitIndex;
+				const uint32 BitIndex = CurrentBitIndex + FPlatformMath::CountTrailingZeros(CurrentWord);
+				Functor(BitIndex);
+				CurrentWord &= CurrentWord - 1;
 			}
 		}
 	}
@@ -952,17 +944,13 @@ private:
 
 			// Flip bits to test the zero values
 			StorageWordType CurrentWord = ~Storage[WordIt] & WordMask;
-			
-			// Test All bits in the CurrentWord and invoke functor if they are set
-			uint32 LocalBitIndex = CurrentBitIndex;
+
+			// Invoke functor for all set bits in the CurrentWord
 			while (CurrentWord)
 			{
-				if (CurrentWord & 0x1)
-				{
-					Functor(LocalBitIndex);
-				}
-				CurrentWord >>= 1;
-				++LocalBitIndex;
+				const uint32 BitIndex = CurrentBitIndex + FPlatformMath::CountTrailingZeros(CurrentWord);
+				Functor(BitIndex);
+				CurrentWord &= CurrentWord - 1;
 			}
 		}
 	}
@@ -979,16 +967,12 @@ private:
 
 			StorageWordType CurrentWord = WordOpFunctor(CurrentWordA, CurrentWordB) & ((WordIt == LastWordIt) ? LastWordMask : ~0U);
 
-			// Test All bits in the CurrentWord and invoke functor if they are set
-			uint32 LocalBitIndex = CurrentBitIndex;
+			// Invoke functor for all set bits in the CurrentWord
 			while (CurrentWord)
 			{
-				if (CurrentWord & 0x1)
-				{
-					Functor(LocalBitIndex);
-				}
-				CurrentWord >>= 1;
-				++LocalBitIndex;
+				const uint32 BitIndex = CurrentBitIndex + FPlatformMath::CountTrailingZeros(CurrentWord);
+				Functor(BitIndex);
+				CurrentWord &= CurrentWord - 1;
 			}
 		}
 	}
@@ -1006,16 +990,12 @@ private:
 
 			StorageWordType CurrentWord = WordOp1(CurrentWordA, WordOp2(CurrentWordB, CurrentWordC) ) & ((WordIt == LastWordIt) ? LastWordMask : ~0U);
 
-			// Test All bits in the CurrentWord and invoke functor if they are set
-			uint32 LocalBitIndex = CurrentBitIndex;
+			// Invoke functor for all set bits in the CurrentWord
 			while (CurrentWord)
 			{
-				if (CurrentWord & 0x1)
-				{
-					Functor(LocalBitIndex);
-				}
-				CurrentWord >>= 1;
-				++LocalBitIndex;
+				const uint32 BitIndex = CurrentBitIndex + FPlatformMath::CountTrailingZeros(CurrentWord);
+				Functor(BitIndex);
+				CurrentWord &= CurrentWord - 1;
 			}
 		}
 	}
