Iris - Don't call NetRPC_GetDebugFunctionName for logging purposes when object and function couldn't be resolved.

#rnx
#Synced-CL 52460203

diff --git a/Engine/Source/Runtime/Net/Iris/Private/Iris/ReplicationSystem/NetBlob/NetRPC.cpp b/Engine/Source/Runtime/Net/Iris/Private/Iris/ReplicationSystem/NetBlob/NetRPC.cpp
--- a/Engine/Source/Runtime/Net/Iris/Private/Iris/ReplicationSystem/NetBlob/NetRPC.cpp
+++ b/Engine/Source/Runtime/Net/Iris/Private/Iris/ReplicationSystem/NetBlob/NetRPC.cpp
@@ -371,10 +371,9 @@ bool FNetRPC::ResolveFunctionAndObject(FNetSerializationContext& Context)
 	const FReplicationStateMemberFunctionDescriptor* FunctionDescriptor = nullptr;
 	if (!NetRPC_GetFunctionAndObject(Context, NetObjectReference, TargetObjectReference, FunctionLocator, FunctionDescriptor, ObjectPtr))
 	{
-		const FString DebugFuncName = NetRPC_GetDebugFunctionName(Context.GetInternalContext()->ReplicationSystem, NetObjectReference, FunctionLocator);
 		UObject* RefObject = NetRPC_GetObject(Context, NetObjectReference, TargetObjectReference);
-		UE_LOGF(LogIrisRpc, Warning, "%s FAILED (GetFunctionAndObject): NetRef=%ls TargetRef=%ls Func=%ls ResolvedObject=%ls FuncLocator=(%u:%u)",
-			__FUNCTION__, *NetObjectReference.ToString(), *TargetObjectReference.ToString(), *DebugFuncName, *GetNameSafe(RefObject),
+		UE_LOGF(LogIrisRpc, Warning, "%s FAILED (GetFunctionAndObject): NetRef=%ls TargetRef=%ls ResolvedObject=%ls FuncLocator=(%u:%u)",
+			__FUNCTION__, *NetObjectReference.ToString(), *TargetObjectReference.ToString(), *GetNameSafe(RefObject),
 			FunctionLocator.DescriptorIndex, FunctionLocator.FunctionIndex);
 		return false;
 	}
