Iris - Made the concept of a local object more explicit in the NetRefHandleManager.

#Synced-CL 52460203

diff --git a/Engine/Source/Runtime/Net/Iris/Private/Iris/ReplicationSystem/NetRefHandleManager.cpp b/Engine/Source/Runtime/Net/Iris/Private/Iris/ReplicationSystem/NetRefHandleManager.cpp
--- a/Engine/Source/Runtime/Net/Iris/Private/Iris/ReplicationSystem/NetRefHandleManager.cpp
+++ b/Engine/Source/Runtime/Net/Iris/Private/Iris/ReplicationSystem/NetRefHandleManager.cpp
@@ -443,22 +443,14 @@ FNetRefHandle FNetRefHandleManager::CreateNetObject(FNetRefHandle WantedHandle,
 	if (InternalIndex != InvalidInternalNetRefIndex)
 	{
 		FReplicatedObjectData& Data = ReplicatedObjectData[InternalIndex];
+		Data.bIsLocal = 1;
 
 		const UE::Net::FReplicationProtocol* ReplicationProtocol = Params.ReplicationProtocol;
-		
-		// Allocate storage for outgoing data. We need a valid pointer even if the size is zero.
-		uint8* StateBuffer = (uint8*)FMemory::MallocZeroed(FPlatformMath::Max(ReplicationProtocol->InternalTotalSize, 1U), ReplicationProtocol->InternalTotalAlignment);
-		if (EnumHasAnyFlags(ReplicationProtocol->ProtocolTraits, EReplicationProtocolTraits::HasConditionalChangeMask))
-		{
-			// Enable all conditions by default. This is to be compatible with the implementation of FRepChangedPropertyTracker where we have hooks into ReplicationSystem.
-			FNetBitArrayView ConditionalChangeMask(reinterpret_cast<uint32*>(StateBuffer + ReplicationProtocol->GetConditionalChangeMaskOffset()), ReplicationProtocol->ChangeMaskBitCount, FNetBitArrayView::NoResetNoValidate);
-			ConditionalChangeMask.SetAllBits();
-		}
-
-		ReplicatedObjectStateBuffers[InternalIndex] = StateBuffer;
-
-		// Bump protocol refcount
+		// We're dependent on the protocol staying valid for the lifetime of the object.
 		ReplicationProtocol->AddRef();
+		
+		// Allocate storage for outgoing data.
+		ReplicatedObjectStateBuffers[InternalIndex] = AllocateStateBuffer(EStateBufferType::Outgoing, ReplicationProtocol);
 
 		return NetRefHandle;
 	}
@@ -485,19 +477,12 @@ FNetRefHandle FNetRefHandleManager::CreateNetObjectFromRemote(FNetRefHandle Want
 	{
 		FReplicatedObjectData& Data = ReplicatedObjectData[InternalIndex];
 
-		// Allocate storage for incoming data
-		Data.ReceiveStateBuffer = (uint8*)FMemory::Malloc(Params.ReplicationProtocol->InternalTotalSize, Params.ReplicationProtocol->InternalTotalAlignment);
-
-		// Since we currently do not have any default values, just initialize to zero
-		// N.B. If this Memzero is optimized away for some reason it still needs to be done for protocols with dynamic state.
-		FMemory::Memzero(Data.ReceiveStateBuffer, Params.ReplicationProtocol->InternalTotalSize);
-
-		// Note: We could initialize this from default but at the moment it is part of the contract for all serializers to write the value when we serialize and we will only apply dirty states
-
-		// Don't bother initializing the conditional changemask if present since it's currently unused on the receiving end.
+		const UE::Net::FReplicationProtocol* ReplicationProtocol = Params.ReplicationProtocol;
+		// We're dependent on the protocol staying valid for the lifetime of the object.
+		ReplicationProtocol->AddRef();
 
-		// Bump protocol refcount
-		Params.ReplicationProtocol->AddRef();
+		// Allocate storage for incoming data
+		Data.ReceiveStateBuffer = AllocateStateBuffer(EStateBufferType::Incoming, ReplicationProtocol);
 
 		return NetRefHandle;
 	}
@@ -1418,6 +1403,19 @@ void FNetRefHandleManager::SetShouldPropagateChangedStates(FNetRefHandle Handle,
 	return SetShouldPropagateChangedStates(ObjectInternalIndex, bShouldPropagateChangedStates);
 }
 
+uint8* FNetRefHandleManager::AllocateStateBuffer(FNetRefHandleManager::EStateBufferType BufferType, const FReplicationProtocol* Protocol)
+{
+	uint8* StateBuffer = static_cast<uint8*>(FMemory::MallocZeroed(FPlatformMath::Max(Protocol->InternalTotalSize, 1U), Protocol->InternalTotalAlignment));
+	if (BufferType == EStateBufferType::Outgoing && EnumHasAnyFlags(Protocol->ProtocolTraits, EReplicationProtocolTraits::HasConditionalChangeMask))
+	{
+		// Enable all conditions by default. This is to be compatible with the implementation of FRepChangedPropertyTracker where we have hooks into ReplicationSystem.
+		FNetBitArrayView ConditionalChangeMask(reinterpret_cast<uint32*>(StateBuffer + Protocol->GetConditionalChangeMaskOffset()), Protocol->ChangeMaskBitCount, FNetBitArrayView::NoResetNoValidate);
+		ConditionalChangeMask.SetAllBits();
+	}
+
+	return StateBuffer;
+}
+
 void FNetRefHandleManager::AddReferencedObjects(FReferenceCollector& Collector)
 {
 	for (TObjectPtr<UObject>& Object : ReplicatedInstances)
diff --git a/Engine/Source/Runtime/Net/Iris/Private/Iris/ReplicationSystem/NetRefHandleManager.h b/Engine/Source/Runtime/Net/Iris/Private/Iris/ReplicationSystem/NetRefHandleManager.h
--- a/Engine/Source/Runtime/Net/Iris/Private/Iris/ReplicationSystem/NetRefHandleManager.h
+++ b/Engine/Source/Runtime/Net/Iris/Private/Iris/ReplicationSystem/NetRefHandleManager.h
@@ -136,6 +136,7 @@ public:
 		uint16 bIsDebugObject : 1U = 0;
 		uint16 bNetObjectIsPendingDestroy : 1U = 0;
 		uint16 bIsDestructionInfo : 1U = 0;
+		uint16 bIsLocal : 1U = 0;
 
 		/** The factory responsible for instantiating this object */
 		FNetObjectFactoryId NetFactoryId = InvalidNetObjectFactoryId;
@@ -495,6 +496,11 @@ public:
 	};
 
 private:
+	enum class EStateBufferType : unsigned
+	{
+		Outgoing,
+		Incoming,
+	};
 
 	void InitNetBitArray(FNetBitArray* NetBitArray);
 
@@ -531,6 +537,8 @@ private:
 	/** Deal with the increase to the max index and tell other systems about it */
 	void MaxInternalNetRefIndexIncreased(uint32 NewBitCount);
 
+	static uint8* AllocateStateBuffer(EStateBufferType BufferType, const FReplicationProtocol* Protocol);
+
 private:
 
 	// Max allowed replicated object count
@@ -729,7 +737,7 @@ bool FNetRefHandleManager::IsNetRefHandleAssigned(FNetRefHandle Handle) const
 
 bool FNetRefHandleManager::IsLocal(FInternalNetRefIndex InternalIndex) const
 {
-	return (InternalIndex != InvalidInternalNetRefIndex) && (ReplicatedObjectStateBuffers[InternalIndex] != nullptr);
+	return GetReplicatedObjectDataNoCheck(InternalIndex).bIsLocal;
 }
 
 bool FNetRefHandleManager::IsLocalNetRefHandle(FNetRefHandle Handle) const
@@ -741,8 +749,7 @@ bool FNetRefHandleManager::IsRemoteNetRefHandle(FNetRefHandle Handle) const
 {
 	if (const uint32 InternalIndex = GetInternalIndex(Handle))
 	{
-		// For the time being only replicated objects owned by this peer has a state buffer
-		return ReplicatedObjectStateBuffers[InternalIndex] == nullptr;
+		return !IsLocal(InternalIndex);
 	}
 
 	return false;
