Iris - Added the notion of ENetConnectionTraits. Current values are None, ReplayRecording and ReplayPlayback. Lifetime conditionals now handle replay conditions based on these traits.

#Synced-CL 52460203

diff --git a/Engine/Source/Runtime/Net/Core/Public/Net/Core/Connection/NetEnums.h b/Engine/Source/Runtime/Net/Core/Public/Net/Core/Connection/NetEnums.h
--- a/Engine/Source/Runtime/Net/Core/Public/Net/Core/Connection/NetEnums.h
+++ b/Engine/Source/Runtime/Net/Core/Public/Net/Core/Connection/NetEnums.h
@@ -3,6 +3,7 @@
 #pragma once
 
 #include "UObject/ObjectMacros.h"
+#include "Misc/EnumClassFlags.h"
 
 #include "NetEnums.generated.h"
 
@@ -123,6 +124,18 @@ namespace ESecurityEvent
 namespace UE::Net
 {
 
+/** Traits of a UNetConnection that are of importance to propagate to replication systems not depending on Engine. */
+enum class ENetConnectionTraits : uint32
+{
+	/** No traits. Must be zero. */
+	None = 0U,
+	/** The connection is for replay recording */
+	ReplayRecording = 1U,
+	/** The connection is for replay playback */
+	ReplayPlayback = 2U,
+};
+ENUM_CLASS_FLAGS(ENetConnectionTraits);
+
 enum class EProxyType : uint8
 {
 	None = 0,
diff --git a/Engine/Source/Runtime/Net/Iris/Private/Iris/ReplicationSystem/Conditionals/ReplicationConditionals.cpp b/Engine/Source/Runtime/Net/Iris/Private/Iris/ReplicationSystem/Conditionals/ReplicationConditionals.cpp
--- a/Engine/Source/Runtime/Net/Iris/Private/Iris/ReplicationSystem/Conditionals/ReplicationConditionals.cpp
+++ b/Engine/Source/Runtime/Net/Iris/Private/Iris/ReplicationSystem/Conditionals/ReplicationConditionals.cpp
@@ -16,6 +16,7 @@
 #include "Iris/ReplicationSystem/ReplicationWriter.h"
 #include "Containers/ArrayView.h"
 #include "UObject/CoreNetTypes.h"
+#include "Net/Core/Connection/NetEnums.h"
 #include "Net/Core/NetHandle/NetHandleManager.h"
 #include "Net/Core/PropertyConditions/RepChangedPropertyTracker.h"
 #include "Net/Core/PropertyConditions/PropertyConditions.h"
@@ -70,7 +71,7 @@ void FReplicationConditionals::OnMaxInternalNetRefIndexIncreased(FInternalNetRef
 
 	for (FPerConnectionInfo& ConnectionInfo : ConnectionInfos)
 	{
-		// Resize the netobject list for valid connections
+		// Resize the ObjectConditionals list for valid connections
 		if (!ConnectionInfo.ObjectConditionals.IsEmpty())
 		{
 			ConnectionInfo.ObjectConditionals.SetNumZeroed(NewMaxInternalIndex);
@@ -175,11 +176,12 @@ void FReplicationConditionals::SetOwningConnection(FInternalNetRefIndex ObjectIn
 	}
 }
 
-void FReplicationConditionals::AddConnection(uint32 ConnectionId)
+void FReplicationConditionals::AddConnection(uint32 ConnectionId, ENetConnectionTraits NetConnectionTraits)
 {
 	// Init connection info
 	FPerConnectionInfo& ConnectionInfo = ConnectionInfos[ConnectionId];
 	ConnectionInfo.ObjectConditionals.SetNumZeroed(MaxInternalNetRefIndex);
+	ConnectionInfo.NetConnectionTraits = NetConnectionTraits;
 }
 
 void FReplicationConditionals::RemoveConnection(uint32 ConnectionId)
@@ -187,6 +189,7 @@ void FReplicationConditionals::RemoveConnection(uint32 ConnectionId)
 	// Reset connection info
 	FPerConnectionInfo& ConnectionInfo = ConnectionInfos[ConnectionId];
 	ConnectionInfo.ObjectConditionals.Empty();
+	ConnectionInfo.NetConnectionTraits = ENetConnectionTraits::None;
 }
 
 void FReplicationConditionals::SetAutonomousAsSimulated(uint32 ConnectionId, bool bEnabled)
@@ -900,7 +903,7 @@ FReplicationConditionals::FConditionalsMask FReplicationConditionals::GetLifetim
 	bool bRoleSimulated = ReplicatingConnectionId != ObjectInfo->AutonomousConnectionId;
 	const bool bRoleAutonomous = ReplicatingConnectionId == ObjectInfo->AutonomousConnectionId;
 	const bool bRepPhysics = ObjectInfo->bRepPhysics;
-
+	const bool bIsReplay = EnumHasAnyFlags(ConnectionInfos[ReplicatingConnectionId].NetConnectionTraits, ENetConnectionTraits::ReplayPlayback | ENetConnectionTraits::ReplayRecording);
 	if (ConnectionInfos.IsValidIndex(ObjectOwnerConnectionId))
 	{
 		const FPerConnectionInfo& ConnectionInfo = ConnectionInfos[ObjectOwnerConnectionId];
@@ -919,11 +922,11 @@ FReplicationConditionals::FConditionalsMask FReplicationConditionals::GetLifetim
 	ConditionalsMask.SetConditionEnabled(COND_SimulatedOrPhysics, bRoleSimulated | bRepPhysics);
 	ConditionalsMask.SetConditionEnabled(COND_InitialOrOwner, bIsReplicatingToOwner | bIsInitialState);
 	ConditionalsMask.SetConditionEnabled(COND_Custom, true);
-	ConditionalsMask.SetConditionEnabled(COND_ReplayOrOwner, bIsReplicatingToOwner);
-	//ConditionalsMask.SetConditionEnabled(COND_ReplayOnly, Not support yet);
-	ConditionalsMask.SetConditionEnabled(COND_SimulatedOnlyNoReplay, bRoleSimulated);
-	ConditionalsMask.SetConditionEnabled(COND_SimulatedOrPhysicsNoReplay, bRoleSimulated | bRepPhysics);
-	ConditionalsMask.SetConditionEnabled(COND_SkipReplay, true);
+	ConditionalsMask.SetConditionEnabled(COND_ReplayOrOwner, bIsReplicatingToOwner | bIsReplay);
+	ConditionalsMask.SetConditionEnabled(COND_ReplayOnly, bIsReplay);
+	ConditionalsMask.SetConditionEnabled(COND_SimulatedOnlyNoReplay, bRoleSimulated & !bIsReplay);
+	ConditionalsMask.SetConditionEnabled(COND_SimulatedOrPhysicsNoReplay, (bRoleSimulated | bRepPhysics) & !bIsReplay);
+	ConditionalsMask.SetConditionEnabled(COND_SkipReplay, !bIsReplay);
 	ConditionalsMask.SetConditionEnabled(COND_Dynamic, true);
 	//ConditionalsMask.SetConditionEnabled(COND_Never, false); Ctor sets its false by default
 	//ConditionalsMask.SetConditionEnabled(COND_NetGroup, Not Known); NetGroup is evaluated in FReplicationConditionals::GetChildSubObjectsToReplicate
diff --git a/Engine/Source/Runtime/Net/Iris/Private/Iris/ReplicationSystem/Conditionals/ReplicationConditionals.h b/Engine/Source/Runtime/Net/Iris/Private/Iris/ReplicationSystem/Conditionals/ReplicationConditionals.h
--- a/Engine/Source/Runtime/Net/Iris/Private/Iris/ReplicationSystem/Conditionals/ReplicationConditionals.h
+++ b/Engine/Source/Runtime/Net/Iris/Private/Iris/ReplicationSystem/Conditionals/ReplicationConditionals.h
@@ -11,6 +11,7 @@
 enum ELifetimeCondition : int;
 namespace UE::Net
 {
+	enum class ENetConnectionTraits : uint32;
 	enum class EReplicationCondition : uint32;
 
 	typedef uint32 FInternalNetRefIndex;
@@ -64,7 +65,7 @@ public:
 	/** Called when when one or more NetRefInternalIndices have been freed and can be re-assigned to new objects. */
 	void OnInternalNetRefIndicesFreed(const TConstArrayView<FInternalNetRefIndex>& FreedIndices);
 
-	void AddConnection(uint32 ConnectionId);
+	void AddConnection(uint32 ConnectionId, ENetConnectionTraits NetConnectionTraits);
 	void RemoveConnection(uint32 ConnectionId);
 
 	/** 
@@ -120,6 +121,7 @@ private:
 	struct FPerConnectionInfo
 	{
 		TArray<FConditionalsMask> ObjectConditionals;
+		ENetConnectionTraits NetConnectionTraits = {};
 		bool bAutonomousAsSimulated = false;
 	};
 
diff --git a/Engine/Source/Runtime/Net/Iris/Private/Iris/ReplicationSystem/DeltaCompression/DeltaCompressionBaselineManager.cpp b/Engine/Source/Runtime/Net/Iris/Private/Iris/ReplicationSystem/DeltaCompression/DeltaCompressionBaselineManager.cpp
--- a/Engine/Source/Runtime/Net/Iris/Private/Iris/ReplicationSystem/DeltaCompression/DeltaCompressionBaselineManager.cpp
+++ b/Engine/Source/Runtime/Net/Iris/Private/Iris/ReplicationSystem/DeltaCompression/DeltaCompressionBaselineManager.cpp
@@ -131,7 +131,7 @@ void FDeltaCompressionBaselineManager::PostSendUpdate(FDeltaCompressionBaselineM
 	DestructBaselineSharingContext(BaselineSharingContext);
 }
 
-void FDeltaCompressionBaselineManager::AddConnection(uint32 ConnectionId)
+void FDeltaCompressionBaselineManager::AddConnection(uint32 ConnectionId, ENetConnectionTraits NetConnectionTraits)
 {
 }
 
diff --git a/Engine/Source/Runtime/Net/Iris/Private/Iris/ReplicationSystem/DeltaCompression/DeltaCompressionBaselineManager.h b/Engine/Source/Runtime/Net/Iris/Private/Iris/ReplicationSystem/DeltaCompression/DeltaCompressionBaselineManager.h
--- a/Engine/Source/Runtime/Net/Iris/Private/Iris/ReplicationSystem/DeltaCompression/DeltaCompressionBaselineManager.h
+++ b/Engine/Source/Runtime/Net/Iris/Private/Iris/ReplicationSystem/DeltaCompression/DeltaCompressionBaselineManager.h
@@ -14,6 +14,7 @@ class UReplicationSystem;
 
 namespace UE::Net
 {
+	enum class ENetConnectionTraits : uint32;
 	enum class ENetObjectDeltaCompressionStatus : unsigned;
 }
 
@@ -72,7 +73,7 @@ public:
 	void PreSendUpdate(FDeltaCompressionBaselineManagerPreSendUpdateParams& UpdateParams);
 	void PostSendUpdate(FDeltaCompressionBaselineManagerPostSendUpdateParams& UpdateParams);
 
-	void AddConnection(uint32 ConnectionId);
+	void AddConnection(uint32 ConnectionId, ENetConnectionTraits NetConnectionTraits);
 	void RemoveConnection(uint32 ConnectionId);
 
 	void SetDeltaCompressionStatus(FInternalNetRefIndex Index, ENetObjectDeltaCompressionStatus Status);
diff --git a/Engine/Source/Runtime/Net/Iris/Private/Iris/ReplicationSystem/Filtering/ReplicationFiltering.cpp b/Engine/Source/Runtime/Net/Iris/Private/Iris/ReplicationSystem/Filtering/ReplicationFiltering.cpp
--- a/Engine/Source/Runtime/Net/Iris/Private/Iris/ReplicationSystem/Filtering/ReplicationFiltering.cpp
+++ b/Engine/Source/Runtime/Net/Iris/Private/Iris/ReplicationSystem/Filtering/ReplicationFiltering.cpp
@@ -639,7 +639,7 @@ FName FReplicationFiltering::GetFilterName(FNetObjectFilterHandle Filter) const
 }
 
 // Connection handling
-void FReplicationFiltering::AddConnection(uint32 ConnectionId)
+void FReplicationFiltering::AddConnection(uint32 ConnectionId, ENetConnectionTraits NetConnectionTraits)
 {
 	bHasNewConnection = 1;
 	ValidConnections.SetBit(ConnectionId);
diff --git a/Engine/Source/Runtime/Net/Iris/Private/Iris/ReplicationSystem/Filtering/ReplicationFiltering.h b/Engine/Source/Runtime/Net/Iris/Private/Iris/ReplicationSystem/Filtering/ReplicationFiltering.h
--- a/Engine/Source/Runtime/Net/Iris/Private/Iris/ReplicationSystem/Filtering/ReplicationFiltering.h
+++ b/Engine/Source/Runtime/Net/Iris/Private/Iris/ReplicationSystem/Filtering/ReplicationFiltering.h
@@ -20,6 +20,7 @@ struct FNetObjectFilterDefinition;
 
 namespace UE::Net
 {
+	enum class ENetConnectionTraits : uint32;
 	typedef uint32 FNetObjectFilterHandle;
 	typedef uint32 FInternalNetRefIndex;
 
@@ -131,7 +132,7 @@ public:
 	void BuildObjectsInFilterList(FNetBitArrayView OutObjectsInFilter, FName FilterName) const;
 
 	// Connection handling
-	void AddConnection(uint32 ConnectionId);
+	void AddConnection(uint32 ConnectionId, ENetConnectionTraits NetConnectionTraits);
 	void RemoveConnection(uint32 ConnectionId);
 
 	// Group based filtering
diff --git a/Engine/Source/Runtime/Net/Iris/Private/Iris/ReplicationSystem/Prioritization/ReplicationPrioritization.cpp b/Engine/Source/Runtime/Net/Iris/Private/Iris/ReplicationSystem/Prioritization/ReplicationPrioritization.cpp
--- a/Engine/Source/Runtime/Net/Iris/Private/Iris/ReplicationSystem/Prioritization/ReplicationPrioritization.cpp
+++ b/Engine/Source/Runtime/Net/Iris/Private/Iris/ReplicationSystem/Prioritization/ReplicationPrioritization.cpp
@@ -614,7 +614,7 @@ void FReplicationPrioritization::Prioritize(const FNetBitArrayView& ReplicatingC
 	}
 }
 
-void FReplicationPrioritization::AddConnection(uint32 ConnectionId)
+void FReplicationPrioritization::AddConnection(uint32 ConnectionId, ENetConnectionTraits NetConnectionTraits)
 {
 	if (ConnectionId >= (uint32)ConnectionInfos.Num())
 	{
diff --git a/Engine/Source/Runtime/Net/Iris/Private/Iris/ReplicationSystem/Prioritization/ReplicationPrioritization.h b/Engine/Source/Runtime/Net/Iris/Private/Iris/ReplicationSystem/Prioritization/ReplicationPrioritization.h
--- a/Engine/Source/Runtime/Net/Iris/Private/Iris/ReplicationSystem/Prioritization/ReplicationPrioritization.h
+++ b/Engine/Source/Runtime/Net/Iris/Private/Iris/ReplicationSystem/Prioritization/ReplicationPrioritization.h
@@ -17,6 +17,7 @@ namespace UE::Net
 	
 	typedef uint32 FInternalNetRefIndex;
 
+	enum class ENetConnectionTraits : uint32;
 	namespace Private
 	{
 		class FNetRefHandleManager;
@@ -56,7 +57,7 @@ public:
 	FNetObjectPrioritizerHandle GetPrioritizerHandle(const FName PrioritizerName) const;
 	UNetObjectPrioritizer* GetPrioritizer(const FName PrioritizerName) const;
 
-	void AddConnection(uint32 ConnectionId);
+	void AddConnection(uint32 ConnectionId, ENetConnectionTraits NetConnectionTraits);
 	void RemoveConnection(uint32 ConnectionId);
 
 	float GetObjectPriorityForConnection(uint32 ConnectionId, FInternalNetRefIndex InternalIndex) const
diff --git a/Engine/Source/Runtime/Net/Iris/Private/Iris/ReplicationSystem/ReplicationConnections.cpp b/Engine/Source/Runtime/Net/Iris/Private/Iris/ReplicationSystem/ReplicationConnections.cpp
--- a/Engine/Source/Runtime/Net/Iris/Private/Iris/ReplicationSystem/ReplicationConnections.cpp
+++ b/Engine/Source/Runtime/Net/Iris/Private/Iris/ReplicationSystem/ReplicationConnections.cpp
@@ -40,6 +40,14 @@ void FReplicationConnections::SetReplicationView(uint32 ConnectionId, const FRep
 	ReplicationViews[ConnectionId] = View;
 }
 
+void FReplicationConnections::AddConnection(uint32 ConnectionId, ENetConnectionTraits NetConnectionTraits)
+{
+	check(ValidConnections.GetBit(ConnectionId) == false);
+	ValidConnections.SetBit(ConnectionId);
+
+	Connections[ConnectionId].NetConnectionTraits = NetConnectionTraits;
+}
+
 void FReplicationConnections::RemoveConnection(uint32 ConnectionId)
 {
 	check(ValidConnections.GetBit(ConnectionId));
@@ -51,6 +59,26 @@ void FReplicationConnections::RemoveConnection(uint32 ConnectionId)
 	ValidConnections.ClearBit(ConnectionId);
 }
 
+bool FReplicationConnections::HasConnectionWithAnyTraits(ENetConnectionTraits NetConnectionTraits) const
+{
+	// For None we  check if there are any valid connections whatsoever.
+	if (NetConnectionTraits == ENetConnectionTraits::None)
+	{
+		return ValidConnections.FindFirstOne() != FNetBitArray::InvalidIndex;
+	}
+
+	for (const FNetBitArray::StorageWordType ConnectionId : ValidConnections)
+	{
+		const FReplicationConnection& Connection = Connections[ConnectionId];
+		if (EnumHasAnyFlags(Connection.NetConnectionTraits, NetConnectionTraits))
+		{
+			return true;
+		}
+	}
+
+	return false;
+}
+
 FNetBitArray FReplicationConnections::GetOpenConnections() const
 {
 	FNetBitArray OpenConnections(ValidConnections.GetNumBits());
diff --git a/Engine/Source/Runtime/Net/Iris/Private/Iris/ReplicationSystem/ReplicationConnections.h b/Engine/Source/Runtime/Net/Iris/Private/Iris/ReplicationSystem/ReplicationConnections.h
--- a/Engine/Source/Runtime/Net/Iris/Private/Iris/ReplicationSystem/ReplicationConnections.h
+++ b/Engine/Source/Runtime/Net/Iris/Private/Iris/ReplicationSystem/ReplicationConnections.h
@@ -2,7 +2,7 @@
 
 #pragma once
 
-#include "CoreTypes.h"
+#include "HAL/Platform.h"
 #include "Containers/Array.h"
 #include "Net/Core/NetBitArray.h"
 #include "Iris/ReplicationSystem/ReplicationView.h"
@@ -10,6 +10,11 @@
 #include "UObject/WeakObjectPtr.h"
 
 class UDataStreamManager;
+namespace UE::Net
+{
+	enum class ENetConnectionTraits : uint32;
+}
+
 namespace UE::Net::Private
 {
 	class FReplicationWriter;
@@ -25,6 +30,7 @@ struct FReplicationConnection
 	FReplicationReader* ReplicationReader = nullptr;
 	TWeakObjectPtr<UDataStreamManager> DataStreamManager;
 	FObjectPtr UserData = nullptr;
+	ENetConnectionTraits NetConnectionTraits = {};
 	bool bIsClosing = false; // Should be set when a connection starts the graceful close process to finish flushing reliable data
 };
 
@@ -70,13 +76,10 @@ public:
 		return ConnectionId < GetMaxConnectionCount() && ValidConnections.GetBit(ConnectionId) && !Connections[ConnectionId].bIsClosing;
 	}
 
-	void AddConnection(uint32 ConnectionId)
-	{
-		check(ValidConnections.GetBit(ConnectionId) == false);
-		ValidConnections.SetBit(ConnectionId);
-	}
+	void AddConnection(uint32 ConnectionId, ENetConnectionTraits NetConnectionTraits);
+	void RemoveConnection(uint32 ConnectionId);
 
-	IRISCORE_API void RemoveConnection(uint32 ConnectionId);
+	bool HasConnectionWithAnyTraits(ENetConnectionTraits NetConnectionTraits) const;
 
 	uint32 GetMaxConnectionCount() const { return ValidConnections.GetNumBits(); }
 
diff --git a/Engine/Source/Runtime/Net/Iris/Private/Iris/ReplicationSystem/ReplicationSystem.cpp b/Engine/Source/Runtime/Net/Iris/Private/Iris/ReplicationSystem/ReplicationSystem.cpp
--- a/Engine/Source/Runtime/Net/Iris/Private/Iris/ReplicationSystem/ReplicationSystem.cpp
+++ b/Engine/Source/Runtime/Net/Iris/Private/Iris/ReplicationSystem/ReplicationSystem.cpp
@@ -71,6 +71,8 @@ static FAutoConsoleVariableRef CVarSetFilterSetsRequiresFrequentLocationUpdates(
 
 }
 
+static_assert(!UE::Net::ENetConnectionTraits::None, "UE::Net::ENetConnectionTraits::None must be zero to avoid header include in ReplicationSystem.h");
+
 namespace UE::Net::Private
 {
 
@@ -699,7 +701,7 @@ public:
 		ReplicationSystemInternal.GetFiltering().DestroyFilter(Name);
 	}
 
-	void AddConnection(uint32 ConnectionId)
+	void AddConnection(uint32 ConnectionId, UE::Net::ENetConnectionTraits NetConnectionTraits)
 	{
 		LLM_SCOPE_BYTAG(IrisConnection);
 
@@ -719,7 +721,7 @@ public:
 		}
 
 		{
-			Connections.AddConnection(ConnectionId);
+			Connections.AddConnection(ConnectionId, NetConnectionTraits);
 
 			FReplicationConnection* Connection = Connections.GetConnection(ConnectionId);
 
@@ -741,6 +743,7 @@ public:
 			// Delaying attachments with unresolved references on the server could cause massive queues of RPCs, potentially an OOM situation.
 			Params.bAllowDelayingAttachmentsWithUnresolvedReferences = !ReplicationSystem->IsServer();
 
+			// $IRIS TODO Replay connection specifics 
 			Connection->ReplicationWriter = new FReplicationWriter();
 			Connection->ReplicationReader = new FReplicationReader();
 
@@ -750,22 +753,22 @@ public:
 
 		{
 			FReplicationConditionals& ReplicationConditionals = ReplicationSystemInternal.GetConditionals();
-			ReplicationConditionals.AddConnection(ConnectionId);
+			ReplicationConditionals.AddConnection(ConnectionId, NetConnectionTraits);
 		}
 
 		{
 			FReplicationFiltering& ReplicationFiltering = ReplicationSystemInternal.GetFiltering();
-			ReplicationFiltering.AddConnection(ConnectionId);
+			ReplicationFiltering.AddConnection(ConnectionId, NetConnectionTraits);
 		}
 
 		{
 			FReplicationPrioritization& ReplicationPrioritization = ReplicationSystemInternal.GetPrioritization();
-			ReplicationPrioritization.AddConnection(ConnectionId);
+			ReplicationPrioritization.AddConnection(ConnectionId, NetConnectionTraits);
 		}
 
 		{
 			FDeltaCompressionBaselineManager& DC = ReplicationSystemInternal.GetDeltaCompressionBaselineManager();
-			DC.AddConnection(ConnectionId);
+			DC.AddConnection(ConnectionId, NetConnectionTraits);
 		}
 
 		FConnectionHandle ConnectionHandle(ConnectionId);
@@ -813,6 +816,12 @@ public:
 		Delegates.ConnectionRemovedDelegate.Broadcast(ConnectionHandle);
 	}
 
+	bool HasConnectionWithAnyTraits(ENetConnectionTraits NetConnectionTraits) const
+	{
+		const FReplicationConnections& Connections = ReplicationSystemInternal.GetConnections();
+		return Connections.HasConnectionWithAnyTraits(NetConnectionTraits);
+	}
+
 	void UpdateDataStreams(UDataStream::FUpdateParameters& UpdateParameters)
 	{
 		IRIS_CSV_PROFILER_SCOPE(Iris, ReplicationSystem_UpdateDataStreams);
@@ -1314,9 +1323,9 @@ void UReplicationSystem::InitializeNetRefHandleGenerator(UReplicationSystem::FRe
 	NetRefHandleGenerator->Bind(BindParams);
 }
 
-void UReplicationSystem::AddConnection(uint32 ConnectionId)
+void UReplicationSystem::AddConnection(uint32 ConnectionId, UE::Net::ENetConnectionTraits Traits)
 {
-	Impl->AddConnection(ConnectionId);
+	Impl->AddConnection(ConnectionId, Traits);
 }
 
 void UReplicationSystem::RemoveConnection(uint32 ConnectionId)
@@ -1324,6 +1333,11 @@ void UReplicationSystem::RemoveConnection(uint32 ConnectionId)
 	Impl->RemoveConnection(ConnectionId);
 }
 
+bool UReplicationSystem::HasConnectionWithAnyTraits(UE::Net::ENetConnectionTraits NetConnectionTraits) const
+{
+	return Impl->HasConnectionWithAnyTraits(NetConnectionTraits);
+}
+
 bool UReplicationSystem::IsValidConnection(uint32 ConnectionId) const
 {
 	UE::Net::Private::FReplicationConnections& Connections = Impl->ReplicationSystemInternal.GetConnections();
diff --git a/Engine/Source/Runtime/Net/Iris/Private/Iris/ReplicationSystem/ReplicationSystemInternal.h b/Engine/Source/Runtime/Net/Iris/Private/Iris/ReplicationSystem/ReplicationSystemInternal.h
--- a/Engine/Source/Runtime/Net/Iris/Private/Iris/ReplicationSystem/ReplicationSystemInternal.h
+++ b/Engine/Source/Runtime/Net/Iris/Private/Iris/ReplicationSystem/ReplicationSystemInternal.h
@@ -92,6 +92,7 @@ public:
 	FChangeMaskCache& GetChangeMaskCache() { return ChangeMaskCache; }
 
 	FReplicationConnections& GetConnections() { return Connections; }
+	const FReplicationConnections& GetConnections() const { return Connections; }
 
 	FReplicationFiltering& GetFiltering() { return Filtering; }
 	const FReplicationFiltering& GetFiltering() const { return Filtering; }
@@ -113,7 +114,7 @@ public:
 
 	FNetTypeStats& GetNetTypeStats() { return TypeStats; }
 
-	FReplicationSystemInternalInitParams& GetInitParams() { return InternalInitParams; }
+	const FReplicationSystemInternalInitParams& GetInitParams() const { return InternalInitParams; }
 
 	FNetSendStats& GetSendStats()
 	{ 
diff --git a/Engine/Source/Runtime/Net/Iris/Public/Iris/ReplicationSystem/ReplicationSystem.h b/Engine/Source/Runtime/Net/Iris/Public/Iris/ReplicationSystem/ReplicationSystem.h
--- a/Engine/Source/Runtime/Net/Iris/Public/Iris/ReplicationSystem/ReplicationSystem.h
+++ b/Engine/Source/Runtime/Net/Iris/Public/Iris/ReplicationSystem/ReplicationSystem.h
@@ -59,6 +59,7 @@ namespace UE::Net
 	struct FNetMetrics;
 	struct FReplicationProtocol;
 
+	enum class ENetConnectionTraits : uint32;
 	enum class ENetFilterStatus : uint32;
 	enum class ENetObjectDeltaCompressionStatus : unsigned;
 	enum class ENetRefHandleError : uint32;
@@ -263,8 +264,9 @@ public:
 	/**
 	 * Notify that a connection was added.
 	 * @param ConnectionId The ID of the added connection. Must not collide with an existing ID. Must be <= the maximum number of connections.
+	 * @param NetConnectionTraits Traits of the connection being added.
 	 */
-	IRISCORE_API void AddConnection(uint32 ConnectionId);
+	IRISCORE_API void AddConnection(uint32 ConnectionId, UE::Net::ENetConnectionTraits NetConnectionTraits = {});
 
 	/**
 	 * Notify that a connection was removed.
@@ -272,6 +274,13 @@ public:
 	 */
 	IRISCORE_API void RemoveConnection(uint32 ConnectionId);
 
+	/**
+	 * Query whether there's a connection with any of the supplied traits. 
+	 * @param NetConnectionTraits The sought after traits of a connection. If UE::Net::ENetConnectionTraits::None is specified it will match all connections.
+	 * @return true if a connection exists with the sought after traits.
+	 */
+	IRISCORE_API bool HasConnectionWithAnyTraits(UE::Net::ENetConnectionTraits NetConnectionTraits) const;
+
 	/**
 	 * Verify if a connection is valid, that is has been added to the system.
 	 * @param ConnectionId The ID of the connection to check for validness.
