Remove deprecated property.

#Synced-CL 52531771

diff --git a/Engine/Source/Runtime/Net/Core/Private/Net/Core/Connection/RPCDoSDetectionConfig.h b/Engine/Source/Runtime/Net/Core/Private/Net/Core/Connection/RPCDoSDetectionConfig.h
--- a/Engine/Source/Runtime/Net/Core/Private/Net/Core/Connection/RPCDoSDetectionConfig.h
+++ b/Engine/Source/Runtime/Net/Core/Private/Net/Core/Connection/RPCDoSDetectionConfig.h
@@ -80,10 +80,6 @@ public:
 	UPROPERTY(config)
 	int32 InitialConnectToleranceMS;
 
-	UE_DEPRECATED(5.1, "This property is no longer supported. Use RPCBlockAllowlist.")
-	UPROPERTY(config)
-	TArray<FName> RPCBlockWhitelist;
-
 	/** List of RPC's which should never be blocked */
 	UPROPERTY(config)
 	TArray<FName> RPCBlockAllowlist;
