[Iris] PolymorphicNetSerializer - Added a cvar to optionally dump the Polymorphic Type Registry when it changes, defaulted to off

#rnx
#Synced-CL 52578497

diff --git a/Engine/Source/Runtime/Net/Iris/Private/Iris/Serialization/PolymorphicNetSerializer.cpp b/Engine/Source/Runtime/Net/Iris/Private/Iris/Serialization/PolymorphicNetSerializer.cpp
--- a/Engine/Source/Runtime/Net/Iris/Private/Iris/Serialization/PolymorphicNetSerializer.cpp
+++ b/Engine/Source/Runtime/Net/Iris/Private/Iris/Serialization/PolymorphicNetSerializer.cpp
@@ -9,12 +9,20 @@
 #include "Iris/ReplicationSystem/ReplicationSystemInternal.h"
 #include "Iris/ReplicationSystem/ReplicationOperationsInternal.h"
 #include "UObject/UObjectIterator.h"
+#include "HAL/IConsoleManager.h"
 
 #include UE_INLINE_GENERATED_CPP_BY_NAME(PolymorphicNetSerializer)
 
 namespace UE::Net
 {
 
+static bool bDumpPolymorphicTypeRegistry = false;
+static FAutoConsoleVariableRef CVarDumpPolymorphicTypeRegistry(
+	TEXT("net.Iris.DumpPolymorphicTypeRegistry"),
+	bDumpPolymorphicTypeRegistry,
+	TEXT("When enabled, logs the full type-to-index mapping for polymorphic net serializer caches whenever they are updated. Useful for diagnosing client/server type index mismatches."),
+	ECVF_Default);
+
 const FName NetError_PolymorphicStructNetSerializer_InvalidStructType("Invalid struct type");
 
 void FPolymorphicNetSerializerScriptStructCache::InitForType(const UScriptStruct* InScriptStruct)
@@ -84,6 +92,20 @@ void FPolymorphicNetSerializerScriptStructCache::InitForType(const UScriptStruct
 	UE_CLOGF(bWasUpdated, LogIris, Log,"FPolymorphicNetSerializerScriptStructCache::InitForType Updated CachedTypedIds for base %ls, NumCachedTypeIDs: %d", ToCStr(InScriptStruct->GetName()), UpdatedRegisteredTypes.Num());
 
 	RegisteredTypes = MoveTemp(UpdatedRegisteredTypes);
+
+	if (bWasUpdated && bDumpPolymorphicTypeRegistry)
+	{
+		UE_LOGF(LogIris, Log, "PolymorphicTypeCache for base %ls: %d types registered", ToCStr(InScriptStruct->GetName()), RegisteredTypes.Num());
+		for (int32 Idx = 0; Idx < RegisteredTypes.Num(); ++Idx)
+		{
+			const FTypeInfo& Info = RegisteredTypes[Idx];
+			UE_LOGF(LogIris, Log, "  [%ls] TypeIndex %u -> %ls (HasObjRef=%ls)",
+				ToCStr(InScriptStruct->GetName()),
+				static_cast<uint32>(Idx + 1),
+				ToCStr(Info.ScriptStruct->GetName()),
+				EnumHasAnyFlags(Info.Descriptor->Traits, EReplicationStateTraits::HasObjectReference) ? TEXT("Yes") : TEXT("No"));
+		}
+	}
 }
 
 }
