[verse] predicts perf: replace temporary net relevancy polling loop with delta callback

- delete UpdateNetRelevancyWatchList and its per-frame ticker
- add OnConnectionScopeChanged delegate to Iris; fires after filtering with per-connection in/out-of-scope bit deltas (word-level XOR+AND, zero cost when unbound)
- predicts subscribes and calls ServerSetObjectRelevancy for each tracked object that changed scope
- replace ServerIsNetRelevant (per-client per-pass) with ServerGetIrisRelevancyInfo (once per object; returns AlwaysRelevant/NeverRelevant/Tracked/Pending)

#Synced-CL 52597118

diff --git a/Engine/Source/Runtime/Net/Iris/Private/Iris/ReplicationSystem/Filtering/ReplicationFiltering.cpp b/Engine/Source/Runtime/Net/Iris/Private/Iris/ReplicationSystem/Filtering/ReplicationFiltering.cpp
--- a/Engine/Source/Runtime/Net/Iris/Private/Iris/ReplicationSystem/Filtering/ReplicationFiltering.cpp
+++ b/Engine/Source/Runtime/Net/Iris/Private/Iris/ReplicationSystem/Filtering/ReplicationFiltering.cpp
@@ -18,6 +18,7 @@
 #include "Iris/ReplicationSystem/ReplicationConnections.h"
 #include "Iris/ReplicationSystem/ReplicationProtocol.h"
 #include "Iris/ReplicationSystem/ReplicationSystem.h"
+#include "Iris/ReplicationSystem/ReplicationSystemDelegates.h"
 #include "Iris/ReplicationSystem/ReplicationSystemInternal.h"
 #include "Iris/ReplicationSystem/Filtering/NetObjectGroups.h"
 #include "Iris/ReplicationSystem/Filtering/NetObjectFilterDefinitions.h"
@@ -254,6 +255,10 @@ void FReplicationFiltering::SetNetObjectListsSize(FInternalNetRefIndex MaxIntern
 		ObjectsWithOwnerFilter.SetNumBits(MaxInternalIndex);
 		ObjectsWithPerObjectInfo.SetNumBits(MaxInternalIndex);
 
+		ScopeDeltaPrevObjects.SetNumBits(MaxInternalIndex);
+		ScopeDeltaNewlyInObjects.SetNumBits(MaxInternalIndex);
+		ScopeDeltaNewlyOutObjects.SetNumBits(MaxInternalIndex);
+
 		DynamicFilterEnabledObjects.SetNumBits(MaxInternalIndex);
 		
 		ObjectsRequiringDynamicFilterUpdate.SetNumBits(MaxInternalIndex);
@@ -1563,6 +1568,15 @@ void FReplicationFiltering::UpdateDynamicFiltering()
 	{
 		FPerConnectionInfo& ConnectionInfo = this->ConnectionInfos[ConnId];
 
+		const bool bTrackScopeDeltas = ReplicationSystem->GetDelegates().ConnectionObjectScopeChangedDelegate.IsBound();
+		if (bTrackScopeDeltas)
+		{
+			FMemory::Memcpy(
+				ScopeDeltaPrevObjects.GetDataChecked(WordCountForObjectBitArrays),
+				ConnectionInfo.ObjectsInScope.GetDataChecked(WordCountForObjectBitArrays),
+				WordCountForObjectBitArrays * sizeof(uint32));
+		}
+
 		ConnectionInfo.InProgressDynamicFilteredOutObjects.ClearAllBits();
 		uint32* InProgressDynamicFilteredOutObjectsData = HasDynamicFilters() ? ConnectionInfo.InProgressDynamicFilteredOutObjects.GetDataChecked(WordCountForObjectBitArrays) : nullptr;
 
@@ -1973,6 +1987,34 @@ void FReplicationFiltering::UpdateDynamicFiltering()
 				}
 			}
 		}
+
+		if (bTrackScopeDeltas)
+		{
+			const uint32* PrevData = ScopeDeltaPrevObjects.GetDataChecked(WordCountForObjectBitArrays);
+			const uint32* CurrData = ConnectionInfo.ObjectsInScope.GetDataChecked(WordCountForObjectBitArrays);
+			uint32* InData = ScopeDeltaNewlyInObjects.GetDataChecked(WordCountForObjectBitArrays);
+			uint32* OutData = ScopeDeltaNewlyOutObjects.GetDataChecked(WordCountForObjectBitArrays);
+			uint32 AnyDiff = 0;
+			for (uint32 WordIt = 0; WordIt < WordCountForObjectBitArrays; ++WordIt)
+			{
+				const uint32 Diff = PrevData[WordIt] ^ CurrData[WordIt];
+				InData[WordIt] = Diff & CurrData[WordIt];
+				OutData[WordIt] = Diff & PrevData[WordIt];
+				AnyDiff |= Diff;
+			}
+
+			if (AnyDiff != 0)
+			{
+				ReplicationSystem->GetDelegates().ConnectionObjectScopeChangedDelegate.Broadcast(
+					FConnectionHandle(ConnId),
+					MakeNetBitArrayView(ScopeDeltaNewlyInObjects),
+					MakeNetBitArrayView(ScopeDeltaNewlyOutObjects),
+					[this](FInternalNetRefIndex Index) -> UObject*
+					{
+						return NetRefHandleManager->GetReplicatedObjectInstance(Index);
+					});
+			}
+		}
 	}
 }
 
diff --git a/Engine/Source/Runtime/Net/Iris/Private/Iris/ReplicationSystem/Filtering/ReplicationFiltering.h b/Engine/Source/Runtime/Net/Iris/Private/Iris/ReplicationSystem/Filtering/ReplicationFiltering.h
--- a/Engine/Source/Runtime/Net/Iris/Private/Iris/ReplicationSystem/Filtering/ReplicationFiltering.h
+++ b/Engine/Source/Runtime/Net/Iris/Private/Iris/ReplicationSystem/Filtering/ReplicationFiltering.h
@@ -383,6 +383,11 @@ private:
 	FInternalNetRefIndex MaxInternalNetRefIndex = 0;
 	uint32 WordCountForObjectBitArrays = 0;
 
+	// Scratch buffers for scope delta tracking (shared across connections, reused each iteration)
+	FNetBitArray ScopeDeltaPrevObjects;
+	FNetBitArray ScopeDeltaNewlyInObjects;
+	FNetBitArray ScopeDeltaNewlyOutObjects;
+
 	// Filter specifics
 	FNetBitArray ObjectsWithDirtyOwnerFilter;
 	FNetBitArray ObjectsWithDirtyOwner;
diff --git a/Engine/Source/Runtime/Net/Iris/Public/Iris/ReplicationSystem/ReplicationSystemDelegates.h b/Engine/Source/Runtime/Net/Iris/Public/Iris/ReplicationSystem/ReplicationSystemDelegates.h
--- a/Engine/Source/Runtime/Net/Iris/Public/Iris/ReplicationSystem/ReplicationSystemDelegates.h
+++ b/Engine/Source/Runtime/Net/Iris/Public/Iris/ReplicationSystem/ReplicationSystemDelegates.h
@@ -4,21 +4,32 @@
 
 #include "HAL/Platform.h"
 #include "Delegates/Delegate.h"
+#include "Templates/Function.h"
+#include "Net/Core/NetBitArray.h"
 #include "Net/Core/Connection/ConnectionHandle.h"
 
 namespace UE::Net::Private
 {
 	class FReplicationSystemImpl;
+	class FReplicationFiltering;
 }
 
 namespace UE::Net
 {
 
+typedef uint32 FInternalNetRefIndex;
+
 class FReplicationSystemDelegates
 {
 public:
 	using FConnectionAddedDelegate = TMulticastDelegate<void(FConnectionHandle ConnectionHandle)>;
 	using FConnectionRemovedDelegate = TMulticastDelegate<void(FConnectionHandle ConnectionHandle)>;
+	using FResolveObjectFromIndex = TFunctionRef<UObject*(FInternalNetRefIndex)>;
+	using FConnectionObjectScopeChangedDelegate = TMulticastDelegate<void(
+		FConnectionHandle /*ConnectionHandle*/,
+		const FNetBitArrayView& /*NewlyInScope*/,
+		const FNetBitArrayView& /*NewlyOutOfScope*/,
+		FResolveObjectFromIndex /*ResolveObject*/)>;
 
 	/** 
 	 * Returns a delegate registration instance allowing the caller to register their FConnectionAddedDelegate. 
@@ -29,16 +40,25 @@ public:
 
 	/**
 	 * Returns a delegate registration instance allowing the caller to register their FConnectionRemovedDelegate.
-	 * The delegate will be called when a previously successfully added connection is removed via a FReplicationSystem::RemoveConnection call. 
+	 * The delegate will be called when a previously successfully added connection is removed via a FReplicationSystem::RemoveConnection call.
 	 * Currently only parent connections will call the delegates.
 	 */
-	FConnectionAddedDelegate::RegistrationType& OnConnectionRemoved(); 
+	FConnectionAddedDelegate::RegistrationType& OnConnectionRemoved();
+
+	/**
+	 * Returns a delegate registration instance allowing the caller to register their FConnectionObjectScopeChangedDelegate.
+	 * The delegate will be called once per connection per frame after filtering is complete, with bitarrays
+	 * of objects that entered or left scope for that connection since the previous frame.
+	 */
+	FConnectionObjectScopeChangedDelegate::RegistrationType& OnConnectionObjectScopeChanged();
 
 private:
 	friend UE::Net::Private::FReplicationSystemImpl;
+	friend class UE::Net::Private::FReplicationFiltering;
 
 	FConnectionAddedDelegate ConnectionAddedDelegate;
 	FConnectionRemovedDelegate ConnectionRemovedDelegate;
+	FConnectionObjectScopeChangedDelegate ConnectionObjectScopeChangedDelegate;
 };
 
 inline FReplicationSystemDelegates::FConnectionAddedDelegate::RegistrationType& FReplicationSystemDelegates::OnConnectionAdded()
@@ -51,4 +71,9 @@ inline FReplicationSystemDelegates::FConnectionRemovedDelegate::RegistrationType
 	return ConnectionRemovedDelegate;
 }
 
+inline FReplicationSystemDelegates::FConnectionObjectScopeChangedDelegate::RegistrationType& FReplicationSystemDelegates::OnConnectionObjectScopeChanged()
+{
+	return ConnectionObjectScopeChangedDelegate;
+}
+
 }
