Iris - Suppress RPC send failures when shutting down test worlds.

#rnx
#Synced-CL 52609747

diff --git a/Engine/Source/Runtime/Engine/Private/Tests/NetTestHelpers.cpp b/Engine/Source/Runtime/Engine/Private/Tests/NetTestHelpers.cpp
--- a/Engine/Source/Runtime/Engine/Private/Tests/NetTestHelpers.cpp
+++ b/Engine/Source/Runtime/Engine/Private/Tests/NetTestHelpers.cpp
@@ -247,6 +247,7 @@ void FTestWorldInstance::Shutdown()
 
 	if (World)
 	{
+		TGuardConsoleVariable<bool> CVarOverrideRPCSendWarnings(TEXT("net.Iris.EnableRPCSendWarnings"), false);
 		World->EndPlay(EEndPlayReason::EndPlayInEditor);
 	}
 
diff --git a/Engine/Source/Runtime/Net/Iris/Private/Iris/ReplicationSystem/NetBlob/NetBlobManager.cpp b/Engine/Source/Runtime/Net/Iris/Private/Iris/ReplicationSystem/NetBlob/NetBlobManager.cpp
--- a/Engine/Source/Runtime/Net/Iris/Private/Iris/ReplicationSystem/NetBlob/NetBlobManager.cpp
+++ b/Engine/Source/Runtime/Net/Iris/Private/Iris/ReplicationSystem/NetBlob/NetBlobManager.cpp
@@ -28,6 +28,9 @@ static TAutoConsoleVariable<int32> CVarEnableIrisRPCs(TEXT("net.Iris.EnableRPCs"
 static bool bThrottleRPCWarnings = true;
 static FAutoConsoleVariableRef CVarThrottleRPCWarnings(TEXT("net.Iris.ThrottleRPCWarnings"), bThrottleRPCWarnings, TEXT("Only log send failure warnings once per RPC type."));
 
+static bool bEnableRPCSendWarnings = true;
+static FAutoConsoleVariableRef CVarEnableRPCSendWarnings(TEXT("net.Iris.EnableRPCSendWarnings"), bEnableRPCSendWarnings, TEXT("Whether to warn about RPC send failures or not. Defaults to true"));
+
 static bool bAllowRPCsOnDormantObjects = true;
 static FAutoConsoleVariableRef CVarAllowRPCsOnDormantObjects(TEXT("net.Iris.RPC.AllowOnDormantObjects"), bAllowRPCsOnDormantObjects, TEXT("When true allow RPCs to be sent on dormant objects. When false block all RPCs from the moment the dormant change is requested."));
 
@@ -271,7 +274,7 @@ bool FNetBlobManager::GetRPCOwner(FRPCOwner& OutOwnerInfo, const FSendRPCContext
 				bLogRPCFailed = !bWasAlreadyLogged;
 				bWasAlreadyLogged = true;
 			}
-			UE_CLOGF(bLogRPCFailed, LogIris, Warning, "SendRPC %ls for %ls Failed. This rootobject is not yet replicated (RefHandle: %ls Index: %u).",
+			UE_CLOGF(bEnableRPCSendWarnings && bLogRPCFailed, LogIris, Warning, "SendRPC %ls for %ls Failed. This root object is not yet replicated (RefHandle: %ls Index: %u).",
 					ToCStr(Context.Function->GetName()), *GetNameSafe(Context.RootObject), ToCStr(OutOwnerInfo.CallerRef.GetRefHandle().ToString()), OutOwnerInfo.RootObjectIndex);
 		}
 	}
