Iris - Removed NetRPC_GetDebugFunctionName due to it being a source of crashes.

#rnx
#Synced-CL 52609747

diff --git a/Engine/Source/Runtime/Net/Iris/Private/Iris/ReplicationSystem/NetBlob/NetRPC.cpp b/Engine/Source/Runtime/Net/Iris/Private/Iris/ReplicationSystem/NetBlob/NetRPC.cpp
--- a/Engine/Source/Runtime/Net/Iris/Private/Iris/ReplicationSystem/NetBlob/NetRPC.cpp
+++ b/Engine/Source/Runtime/Net/Iris/Private/Iris/ReplicationSystem/NetBlob/NetRPC.cpp
@@ -59,7 +59,6 @@ static UObject* NetRPC_GetRootObject(FNetSerializationContext& Context, const FN
 
 static bool NetRPC_GetFunctionAndObject(FNetSerializationContext& Context, const FNetObjectReference& RootObjectReference, const FNetObjectReference& SubObjectReferece, const FNetRPC::FFunctionLocator& FunctionLocator, const FReplicationStateMemberFunctionDescriptor*& OutFunctionDescriptor, TWeakObjectPtr<UObject>& OutObject);
 
-static FString NetRPC_GetDebugFunctionName(const UReplicationSystem* ReplicationSystem, const FNetObjectReference& ObjectReference, const FNetRPC::FFunctionLocator& FunctionLocator);
 static FString NetRPC_GetDebugObjectRefName(FNetSerializationContext& Context, const FNetObjectReference& ObjectReference, UObject* ObjectPtr = nullptr);
 
 static const FName NetError_InvalidNetObjectReference("Invalid NetObjectRefererence");
@@ -541,7 +540,7 @@ void FNetRPC::CallFunction(FNetRPCCallContext& CallContext)
 
 	if (Object == nullptr || Function == nullptr)
 	{
-		const FString DebugFuncName = Function ? Function->GetName() : NetRPC_GetDebugFunctionName(ReplicationSystem, NetObjectReference, FunctionLocator);
+		const FString DebugFuncName = Function ? Function->GetName() : FString("unknown");
 		if (!TargetObjectReference.IsValid())
 		{
 			UE_LOGF(LogIrisRpc, Error, "Rejected RPC (%u:%u) '%ls' due to missing object or function for object: %ls.",
@@ -572,7 +571,7 @@ void FNetRPC::CallFunction(FNetRPCCallContext& CallContext)
 		if (!NetRefHandleManager.HasInstanceProtocol(InternalIndex))
 		{
 			UE_LOGF(LogIrisRpc, Warning, "Rejected RPC: %ls due to object: %ls having stopped replication",
-				*(Function ? Function->GetName() : NetRPC_GetDebugFunctionName(ReplicationSystem, NetObjectReference, FunctionLocator)),
+				*(Function ? Function->GetName() : FString("unknown")),
 				*NetRPC_GetDebugObjectRefName(Context, NetObjectReference, Object));
 
 			return;
@@ -868,28 +867,6 @@ static bool NetRPC_GetFunctionAndObject(FNetSerializationContext& Context, const
 	return true;
 }
 
-static FString NetRPC_GetDebugFunctionName(const UReplicationSystem* ReplicationSystem, const FNetObjectReference& ObjectReference, const FNetRPC::FFunctionLocator& FunctionLocator)
-{
-	const FReplicationProtocol* Protocol = ReplicationSystem->GetReplicationProtocol(ObjectReference.GetRefHandle());
-	if (!Protocol)
-	{
-		return TEXT("ProtocolNotFound");
-	}
-
-	if (FunctionLocator.DescriptorIndex >= Protocol->ReplicationStateCount)
-	{
-		return TEXT("InvalidDescriptorIndex");
-	}
-
-	const FReplicationStateDescriptor* Descriptor = Protocol->ReplicationStateDescriptors[FunctionLocator.DescriptorIndex];
-	if (FunctionLocator.FunctionIndex >= Descriptor->FunctionCount)
-	{
-		return TEXT("InvalidFunctionIndex");
-	}
-
-	return GetNameSafe(Descriptor->MemberFunctionDescriptors[FunctionLocator.FunctionIndex].Function);
-}
-
 static FString NetRPC_GetDebugObjectRefName(FNetSerializationContext& Context, const FNetObjectReference& ObjectReference, UObject* ObjectPtr)
 {
 	if (!ObjectPtr)
